- Move new spline development to CineCameraRigs plugin
- Create spline component class to store camera data and custom parameterization
- Create CineCameraRigRail class to drive attachment with custom parameterization
- Additional attachment control on CineCameraRigRail
- Speed visualization
- POC widget to show new spline/CineCameraRigRail functionality

diff --git a/Engine/Plugins/Experimental/CineCameraRigs/CineCameraRigs.uplugin b/Engine/Plugins/Experimental/CineCameraRigs/CineCameraRigs.uplugin
new file mode 100644
--- /dev/null
+++ b/Engine/Plugins/Experimental/CineCameraRigs/CineCameraRigs.uplugin
@@ -0,0 +1,35 @@
+{
+	"FileVersion": 3,
+	"Version": 1,
+	"VersionName": "1.0",
+	"FriendlyName": "CineCameraRigs",
+	"Description": "Extended camera rigs for cinematic workflow",
+	"Category": "Virtual Production",
+	"CreatedBy": "Epic Games, Inc.",
+	"CreatedByURL": "https://epicgames.com",
+	"DocsURL": "",
+	"MarketplaceURL": "",
+	"SupportURL": "",
+	"EnabledByDefault": false,
+	"CanContainContent": true,
+	"IsExperimentalVersion": true,
+	"Installed": false,
+	"Modules": [
+		{
+			"Name": "CineCameraRigs",
+			"Type": "Runtime",
+			"LoadingPhase": "Default"
+		},
+		{
+			"Name": "CineCameraRigsEditor",
+			"Type": "Editor",
+			"LoadingPhase": "PostEngineInit"
+		}
+	],
+	"Plugins": [
+		{
+			"Name": "EditorScriptingUtilities",
+			"Enabled": true
+		},
+	]
+}
diff --git a/Engine/Plugins/Experimental/CineCameraRigs/Config/BaseCineCameraRigs.ini b/Engine/Plugins/Experimental/CineCameraRigs/Config/BaseCineCameraRigs.ini
new file mode 100644
--- /dev/null
+++ b/Engine/Plugins/Experimental/CineCameraRigs/Config/BaseCineCameraRigs.ini
@@ -0,0 +1,3 @@
+[/Script/CineCameraRigs.CineCameraRigRailSettings]
+DefaultSplineMeshMaterial=/CineCameraRigs/RigRail/M_RigRail.M_RigRail
+DefaultSplineMeshTexture=/CineCameraRigs/RigRail/T_RigRailDefault.T_RigRailDefault
\ No newline at end of file
diff --git a/Engine/Plugins/Experimental/CineCameraRigs/Source/CineCameraRigs/CineCameraRigs.Build.cs b/Engine/Plugins/Experimental/CineCameraRigs/Source/CineCameraRigs/CineCameraRigs.Build.cs
new file mode 100644
--- /dev/null
+++ b/Engine/Plugins/Experimental/CineCameraRigs/Source/CineCameraRigs/CineCameraRigs.Build.cs
@@ -0,0 +1,58 @@
+// Copyright Epic Games, Inc. All Rights Reserved.
+
+using UnrealBuildTool;
+
+public class CineCameraRigs : ModuleRules
+{
+	public CineCameraRigs(ReadOnlyTargetRules Target) : base(Target)
+	{
+		PCHUsage = ModuleRules.PCHUsageMode.UseExplicitOrSharedPCHs;
+		
+		PublicIncludePaths.AddRange(
+			new string[] {
+				// ... add public include paths required here ...
+			}
+			);
+				
+		
+		PrivateIncludePaths.AddRange(
+			new string[] {
+				// ... add other private include paths required here ...
+			}
+			);
+			
+		
+		PublicDependencyModuleNames.AddRange(
+			new string[]
+			{
+				"Core",
+				"CinematicCamera",
+				// ... add other public dependencies that you statically link with here ...
+			}
+			);
+			
+		
+		PrivateDependencyModuleNames.AddRange(
+			new string[]
+			{
+				"CoreUObject",
+				"Engine",
+				"LevelSequence",
+				"MovieScene",
+				"MovieSceneTracks",
+				"Slate",
+				"SlateCore",
+				// ... add private dependencies that you statically link with here ...	
+			}
+			);
+		
+		
+		DynamicallyLoadedModuleNames.AddRange(
+			new string[]
+			{
+				// ... add any modules that your module loads dynamically here ...
+			}
+			);
+
+	}
+}
diff --git a/Engine/Plugins/Experimental/CineCameraRigs/Source/CineCameraRigs/Private/CineCameraRigRail.cpp b/Engine/Plugins/Experimental/CineCameraRigs/Source/CineCameraRigs/Private/CineCameraRigRail.cpp
new file mode 100644
--- /dev/null
+++ b/Engine/Plugins/Experimental/CineCameraRigs/Source/CineCameraRigs/Private/CineCameraRigRail.cpp
@@ -0,0 +1,300 @@
+// Copyright Epic Games, Inc. All Rights Reserved.
+
+#include "CineCameraRigRail.h"
+#include "CineCameraRigRailHelpers.h"
+#include "CineCameraRigsSettings.h"
+#include "CineSplineComponent.h"
+
+#include "CineCameraComponent.h"
+#include "Components/SplineMeshComponent.h"
+#include "Kismet/KismetMathLibrary.h"
+#include "Materials/MaterialInstanceDynamic.h"
+#include "Engine/World.h"
+#include "UObject/Package.h"
+
+ACineCameraRigRail::ACineCameraRigRail(const FObjectInitializer& ObjectInitializer)
+	: Super(ObjectInitializer
+		.SetDefaultSubobjectClass<UCineSplineComponent>(TEXT("RailSplineComponent"))
+	)
+{
+	CineSplineComponent = Cast<UCineSplineComponent>(GetRailSplineComponent());
+	if (CineSplineComponent != nullptr)
+	{
+		CineSplineComponent->OnSplineEdited.AddUObject(this, &ACineCameraRigRail::OnSplineEdited);
+	}
+
+	UMaterialInterface* DefaultMaterial = GetDefault<UCineCameraRigRailSettings>()->DefaultSplineMeshMaterial.LoadSynchronous();
+	if (DefaultMaterial)
+	{
+		SplineMeshMaterial = DefaultMaterial;
+	}
+
+	UTexture2D* DefaultTexture = GetDefault<UCineCameraRigRailSettings>()->DefaultSplineMeshTexture.LoadSynchronous();
+	if (DefaultTexture)
+	{
+		SplineMeshTexture = DefaultTexture;
+	}
+
+	bLockOrientationToRail = true;
+
+#if WITH_EDITORONLY_DATA
+	if (PreviewMesh_Mount)
+	{
+		PreviewMesh_Mount->bHiddenInGame = false;
+	}
+
+#endif
+}
+
+void ACineCameraRigRail::UpdateRailComponents()
+{
+	if (!GetWorld())
+	{
+		return;
+	}
+
+	if (!bUseCustomPosition)
+	{
+		Super::UpdateRailComponents();
+	}
+	else
+	{
+		USceneComponent* AttachComponent = GetDefaultAttachComponent();
+		if (CineSplineComponent && AttachComponent)
+		{
+			float InputKey = CineSplineComponent->GetInputKeyAtPosition(CustomPosition);
+			FVector const SplinePosition = CineSplineComponent->GetLocationAtSplineInputKey(InputKey, ESplineCoordinateSpace::World);
+			FQuat SplineQuat = CineSplineComponent->GetQuaternionAtSplineInputKey(InputKey, ESplineCoordinateSpace::World);
+
+			if (bUseCameraRotationMetadata)
+			{
+				SplineQuat = CineSplineComponent->GetComponentTransform().GetRotation() * CineSplineComponent->GetCameraRotationAtSplineInputKey(InputKey);
+			}
+			FVector Position = GetActorLocation();
+			FRotator Rotation = GetActorRotation();
+			if (bAttachLocationX)
+			{
+				Position.X = SplinePosition.X;
+			}
+			if (bAttachLocationY)
+			{
+				Position.Y = SplinePosition.Y;
+			}
+			if (bAttachLocationZ)
+			{
+				Position.Z = SplinePosition.Z;
+			}
+			if (bAttachRotationX)
+			{
+				Rotation.Roll = SplineQuat.Rotator().Roll;
+			}
+			if (bAttachRotationY)
+			{
+				Rotation.Pitch = SplineQuat.Rotator().Pitch;
+			}
+			if (bAttachRotationZ)
+			{
+				Rotation.Yaw = SplineQuat.Rotator().Yaw;
+			}
+			if (bLockOrientationToRail)
+			{
+				AttachComponent->SetWorldTransform(FTransform(Rotation, Position));
+			}
+			else
+			{
+				AttachComponent->SetWorldLocation(SplinePosition);
+			}
+		}
+	}
+	TArray< AActor* > AttachedActors;
+	GetAttachedActors(AttachedActors);
+	float InputKey = 0.0f;
+	if (bUseCustomPosition)
+	{
+		InputKey = CineSplineComponent->GetInputKeyAtPosition(CustomPosition);
+	}
+	else
+	{
+		float const SplineLen = RailSplineComponent->GetSplineLength();
+		InputKey = CineSplineComponent->GetInputKeyValueAtDistanceAlongSpline(CurrentPositionOnRail * SplineLen);
+	}
+	float FocalLengthValue = CineSplineComponent->GetFloatPropertyAtSplineInputKey(InputKey, FName("FocalLength"));
+	float ApertureValue = CineSplineComponent->GetFloatPropertyAtSplineInputKey(InputKey, FName("Aperture"));
+	float FocusDistanceValue = CineSplineComponent->GetFloatPropertyAtSplineInputKey(InputKey, FName("FocusDistance"));
+	for (AActor* AttachedActor : AttachedActors)
+	{
+		if (UCineCameraComponent* CameraComponent = Cast<UCineCameraComponent>(AttachedActor->GetComponentByClass(UCineCameraComponent::StaticClass())))
+		{
+			if (bDriveFocusDistance)
+			{
+				CameraComponent->FocusSettings.ManualFocusDistance = FocusDistanceValue;
+			}
+			if (bDriveFocalLength)
+			{
+				CameraComponent->SetCurrentFocalLength(FocalLengthValue);
+			}
+			if (bDriveAperture)
+			{
+				CameraComponent->SetCurrentAperture(ApertureValue);
+			}
+		}
+	}
+#if WITH_EDITOR
+	if (GIsEditor)
+	{
+		const UWorld* const MyWorld = GetWorld();
+		if (MyWorld && !MyWorld->IsGameWorld())
+		{
+			// If bUseCustomPosition is false, UpdatePreviewMeshes() is already called inside Super::UpdateRailComponents()
+			if (bUseCustomPosition)
+			{
+				UpdatePreviewMeshes();
+			}
+		}
+
+		// Set HiddenInGame false on the spline mesh so that it can show up in game mode
+		// It can still be hidden via actor's HiddenInGame property
+		// Also update material if numbers don't match
+		int32 NumSegments = PreviewRailMeshSegments.Num();
+		if (NumSegments != SplineMeshMIDs.Num())
+		{
+			UpdateSplineMeshMID();
+		}
+		for (int Index = 0; Index < NumSegments; ++Index)
+		{
+			USplineMeshComponent* SplineMeshComp = PreviewRailMeshSegments[Index];
+			if (SplineMeshComp)
+			{
+				SplineMeshComp->bHiddenInGame = false;
+				if (SplineMeshMIDs.Num() > Index && IsValid(SplineMeshMIDs[Index]))
+				{
+					SplineMeshComp->SetMaterial(0, SplineMeshMIDs[Index]);
+				}
+			}
+		}
+	}
+#endif
+}
+
+void ACineCameraRigRail::UpdateSplineMeshMID()
+{
+	if (GetWorld())
+	{
+		if (IsValid(SplineMeshMaterial))
+		{
+			int32 const NumSplinePoints = RailSplineComponent->GetNumberOfSplinePoints();
+			if (SplineMeshMIDs.Num() > NumSplinePoints - 1)
+			{
+				int32 const NumToRemove = SplineMeshMIDs.Num() - (NumSplinePoints - 1);
+				for (int32 Index = 0; Index < NumToRemove; ++Index)
+				{
+					SplineMeshMIDs.Pop();
+				}
+			}
+			else
+			{
+				int32 const NumToAdd = (NumSplinePoints - 1) - SplineMeshMIDs.Num();
+				for (int Index = 0; Index < NumToAdd; ++Index)
+				{
+					UMaterialInstanceDynamic* SplineMeshMID = UMaterialInstanceDynamic::Create(SplineMeshMaterial, nullptr, MakeUniqueObjectName(GetTransientPackage(), UMaterialInstanceDynamic::StaticClass(), TEXT("SplineMeshMID")));
+					SplineMeshMID->SetFlags(RF_Transient | RF_NonPIEDuplicateTransient | RF_TextExportTransient);
+					SplineMeshMIDs.Add(SplineMeshMID);
+				}
+			}
+			
+			for(int32 Index = 0; Index < NumSplinePoints - 1; ++Index)
+			{
+				UMaterialInstanceDynamic* SplineMeshMID = SplineMeshMIDs[Index];
+				if (!IsValid(SplineMeshMID) || (SplineMeshMID->Parent != SplineMeshMaterial))
+				{
+					SplineMeshMID = UMaterialInstanceDynamic::Create(SplineMeshMaterial, nullptr, MakeUniqueObjectName(GetTransientPackage(), UMaterialInstanceDynamic::StaticClass(), TEXT("SplineMeshMID")));
+					SplineMeshMID->SetFlags(RF_Transient | RF_NonPIEDuplicateTransient | RF_TextExportTransient);
+					SplineMeshMIDs[Index] = SplineMeshMID;
+				}
+			}
+		}
+	}
+}
+
+#if WITH_EDITOR
+void ACineCameraRigRail::PostEditChangeProperty(FPropertyChangedEvent& PropertyChangedEvent)
+{
+	Super::PostEditChangeProperty(PropertyChangedEvent);
+	const FName PropertyName = PropertyChangedEvent.GetPropertyName();
+
+	if (PropertyName == GET_MEMBER_NAME_CHECKED(ACineCameraRigRail, SplineMeshMaterial) ||
+		PropertyName == GET_MEMBER_NAME_CHECKED(ACineCameraRigRail, SplineMeshTexture))
+	{
+			UpdateSplineMeshMID();
+			SetMIDParameters();
+	}
+}
+#endif
+
+void ACineCameraRigRail::SetSplineMeshMaterial(UMaterialInterface* InMaterial)
+{
+	SplineMeshMaterial = InMaterial;
+	UpdateSplineMeshMID();
+	SetMIDParameters();
+}
+void ACineCameraRigRail::SetSplineMeshTexture(UTexture2D* InTexture)
+{
+	SplineMeshTexture = InTexture;
+	UpdateSplineMeshMID();
+	SetMIDParameters();
+}
+
+FVector ACineCameraRigRail::GetVelocityAtPosition(const float InPosition, const float delta) const
+{
+	if (CineSplineComponent->Duration == 0.0f)
+	{
+		return FVector::ZeroVector;
+	}
+
+	float TimeMultiplier = 1.0f / CineSplineComponent->Duration;
+
+	if (bUseCustomPosition)
+	{
+		int32 const NumPoints = CineSplineComponent->GetNumberOfSplinePoints();
+		float const MinPosition = CineSplineComponent->GetFloatPropertyAtSplinePoint(0, FName(TEXT("CustomPosition")));
+		float const MaxPosition = CineSplineComponent->GetFloatPropertyAtSplinePoint(NumPoints - 1, FName(TEXT("CustomPosition")));
+		float const t0 = InPosition + delta >= MaxPosition ? InPosition - delta : InPosition;
+		float const t1 = InPosition + delta >= MaxPosition ? InPosition : InPosition + delta;
+		float const InKey0 = CineSplineComponent->GetInputKeyAtPosition(t0);
+		float const InKey1 = CineSplineComponent->GetInputKeyAtPosition(t1);
+		FVector const P0 = CineSplineComponent->GetLocationAtSplineInputKey(InKey0, ESplineCoordinateSpace::World);
+		FVector const P1 = CineSplineComponent->GetLocationAtSplineInputKey(InKey1, ESplineCoordinateSpace::World);
+		return (P1 - P0) / (t1 - t0) * TimeMultiplier * (MaxPosition - MinPosition);
+	}
+	else
+	{
+		float const SplineLen = CineSplineComponent->GetSplineLength();
+		float const t0 = InPosition + delta >= 1.0f ? InPosition - delta : InPosition;
+		float const t1 = InPosition + delta >= 1.0f ? InPosition : InPosition + delta;
+		FVector const P0 = CineSplineComponent->GetLocationAtDistanceAlongSpline(t0 * SplineLen, ESplineCoordinateSpace::World);
+		FVector const P1 = CineSplineComponent->GetLocationAtDistanceAlongSpline(t1 * SplineLen, ESplineCoordinateSpace::World);
+		return (P1 - P0) / (t1 - t0) * TimeMultiplier;
+	}
+	return FVector::ZeroVector;
+}
+
+void ACineCameraRigRail::OnSplineEdited()
+{
+	UpdateSplineMeshMID();
+	SetMIDParameters();
+}
+
+void ACineCameraRigRail::SetMIDParameters()
+{
+	int32 NumMIDs = SplineMeshMIDs.Num();
+	for (int32 Index = 0; Index < NumMIDs; ++Index)
+	{
+		UMaterialInstanceDynamic* MID = SplineMeshMIDs[Index];
+		MID->SetScalarParameterValue(TEXT("SplineParamStart"), UKismetMathLibrary::SafeDivide((float)Index, (float)NumMIDs));
+		MID->SetScalarParameterValue(TEXT("SplineParamEnd"), UKismetMathLibrary::SafeDivide((float)Index+1, (float)NumMIDs));
+		if (IsValid(SplineMeshTexture))
+		{
+			MID->SetTextureParameterValue(TEXT("SplineTexture"), SplineMeshTexture);
+		}
+	}
+}
diff --git a/Engine/Plugins/Experimental/CineCameraRigs/Source/CineCameraRigs/Private/CineCameraRigRailHelpers.cpp b/Engine/Plugins/Experimental/CineCameraRigs/Source/CineCameraRigs/Private/CineCameraRigRailHelpers.cpp
new file mode 100644
--- /dev/null
+++ b/Engine/Plugins/Experimental/CineCameraRigs/Source/CineCameraRigs/Private/CineCameraRigRailHelpers.cpp
@@ -0,0 +1,45 @@
+// Copyright Epic Games, Inc. All Rights Reserved.
+
+#include "CineCameraRigRailHelpers.h"
+#include "CineCameraRigRail.h"
+
+#include "Kismet/KismetMathLibrary.h"
+#include "Engine/TextureDefines.h"
+#include "Engine/Texture.h"
+#include "TextureResource.h"
+
+void UCineCameraRigRailHelpers::CreateOrUpdateSplineHeatmapTexture(UTexture2D*& Texture, const TArray<float>& DataValues, const float LowValue, const float MidValue, const float HighValue)
+{
+#if WITH_EDITOR
+	int32 Width = DataValues.Num();
+	int32 Height = 1;
+	TArray<FColor> Pixels;
+	Pixels.Init(FColor::Black, Height * Width);
+	float HighRange = HighValue - MidValue;
+	for (int32 Index = 0; Index < Width; ++Index)
+	{
+		float UpperWeight = UKismetMathLibrary::SafeDivide(FMath::Clamp(DataValues[Index] - MidValue, 0.0f, HighValue - MidValue), HighValue - MidValue);
+		float LowerWeight = UKismetMathLibrary::SafeDivide(FMath::Clamp(MidValue - DataValues[Index], 0.0f, MidValue - LowValue), MidValue - LowValue);
+		Pixels[Index].R = (uint8)(UpperWeight * 255);
+		Pixels[Index].G = (uint8)((1.0 - (UpperWeight + LowerWeight)) * 255);
+		Pixels[Index].B = (uint8)(LowerWeight * 255);
+		Pixels[Index].A = DataValues[Index] >= LowValue ? 255 : 0;
+	}
+
+	if (Texture == nullptr || Texture->GetSizeX() != Width || Texture->GetSizeY() != Height)
+	{
+		Texture = UTexture2D::CreateTransient(Width, Height, EPixelFormat::PF_B8G8R8A8);
+		Texture->SRGB = false;
+		Texture->Filter = TextureFilter::TF_Bilinear;
+		Texture->MipGenSettings = TextureMipGenSettings::TMGS_NoMipmaps;
+		Texture->CompressionSettings = TextureCompressionSettings::TC_VectorDisplacementmap;
+		Texture->UpdateResource();
+	}
+
+	void* RawTextureData = Texture->GetPlatformData()->Mips[0].BulkData.Lock(LOCK_READ_WRITE);
+	FMemory::Memcpy(RawTextureData, Pixels.GetData(), Pixels.Num() * sizeof(FColor));
+	Texture->GetPlatformData()->Mips[0].BulkData.Unlock();
+	Texture->UpdateResource();
+#endif
+}
+
diff --git a/Engine/Plugins/Experimental/CineCameraRigs/Source/CineCameraRigs/Private/CineCameraRigs.cpp b/Engine/Plugins/Experimental/CineCameraRigs/Source/CineCameraRigs/Private/CineCameraRigs.cpp
new file mode 100644
--- /dev/null
+++ b/Engine/Plugins/Experimental/CineCameraRigs/Source/CineCameraRigs/Private/CineCameraRigs.cpp
@@ -0,0 +1,23 @@
+// Copyright Epic Games, Inc. All Rights Reserved.
+
+#include "CineCameraRigs.h"
+#include "CineSplineLog.h"
+
+#define LOCTEXT_NAMESPACE "FCineCameraRigsModule"
+
+DEFINE_LOG_CATEGORY(LogCineSpline);
+
+void FCineCameraRigsModule::StartupModule()
+{
+	// This code will execute after your module is loaded into memory; the exact timing is specified in the .uplugin file per-module
+}
+
+void FCineCameraRigsModule::ShutdownModule()
+{
+	// This function may be called during shutdown to clean up your module.  For modules that support dynamic reloading,
+	// we call this function before unloading the module.
+}
+
+#undef LOCTEXT_NAMESPACE
+	
+IMPLEMENT_MODULE(FCineCameraRigsModule, CineCameraRigs)
\ No newline at end of file
diff --git a/Engine/Plugins/Experimental/CineCameraRigs/Source/CineCameraRigs/Private/CineSplineComponent.cpp b/Engine/Plugins/Experimental/CineCameraRigs/Source/CineCameraRigs/Private/CineSplineComponent.cpp
new file mode 100644
--- /dev/null
+++ b/Engine/Plugins/Experimental/CineCameraRigs/Source/CineCameraRigs/Private/CineSplineComponent.cpp
@@ -0,0 +1,327 @@
+// Copyright Epic Games, Inc. All Rights Reserved.
+
+#include "CineSplineComponent.h"
+
+#include "Engine/Engine.h"
+
+UCineSplineComponent::UCineSplineComponent(const FObjectInitializer& ObjectInitializer)
+	: Super(ObjectInitializer)
+{
+	SplineCurves.Position.Points.Reset(10);
+	SplineCurves.Rotation.Points.Reset(10);
+	SplineCurves.Scale.Points.Reset(10);
+	CineSplineMetadata = ObjectInitializer.CreateDefaultSubobject<UCineSplineMetadata>(this, TEXT("CineSplineMetadata"));
+	CineSplineMetadata->Reset(10);
+
+	FCineSplinePointData Data;
+	FCineSplineCurveDefaults Defaults;
+	Data.Location = FVector(0, 0, 0);
+	Data.Rotation = FRotator(0, 0, 0);
+	Data.FocalLength = Defaults.DefaultFocalLength;
+	Data.Aperture = Defaults.DefaultAperture;
+	Data.FocusDistance = Defaults.DefaultFocusDistance;
+
+	AddSplineDataAtPosition(1.0f, Data);
+}
+
+
+USplineMetadata* UCineSplineComponent::GetSplinePointsMetadata()
+{
+	return CineSplineMetadata;
+}
+
+const USplineMetadata* UCineSplineComponent::GetSplinePointsMetadata() const
+{
+	return CineSplineMetadata;
+}
+
+void UCineSplineComponent::PostLoad()
+{
+	Super::PostLoad();
+	if (CineSplineMetadata)
+	{
+		SynchronizeProperties();
+	}
+}
+
+TStructOnScope<FActorComponentInstanceData> UCineSplineComponent::GetComponentInstanceData() const
+{
+	TStructOnScope<FActorComponentInstanceData> InstanceData = MakeStructOnScope<FActorComponentInstanceData, FCineSplineInstanceData>(this);
+	FCineSplineInstanceData* SplineInstanceData = InstanceData.Cast<FCineSplineInstanceData>();
+
+	if (bSplineHasBeenEdited)
+	{
+		SplineInstanceData->CineSplineMetadata = CineSplineMetadata;
+		SplineInstanceData->SplineCurves = SplineCurves;
+	}
+
+	SplineInstanceData->bSplineHasBeenEdited = bSplineHasBeenEdited;
+
+	return InstanceData;
+}
+
+void UCineSplineComponent::ApplyComponentInstanceData(FCineSplineInstanceData* SplineInstanceData, const bool bPostUCS)
+{
+	check(SplineInstanceData);
+
+	if (bPostUCS)
+	{
+		if (bInputSplinePointsToConstructionScript)
+		{
+			// Don't reapply the saved state after the UCS has run if we are inputting the points to it.
+			// This allows the UCS to work on the edited points and make its own changes.
+			return;
+		}
+		else
+		{
+			//bModifiedByConstructionScript = (SplineInstanceData->SplineCurvesPreUCS != SplineCurves);
+
+			// If we are restoring the saved state, unmark the SplineCurves property as 'modified'.
+			// We don't want to consider that these changes have been made through the UCS.
+			TArray<FProperty*> Properties;
+			Properties.Emplace(FindFProperty<FProperty>(UCineSplineComponent::StaticClass(), GET_MEMBER_NAME_CHECKED(UCineSplineComponent, SplineCurves)));
+			RemoveUCSModifiedProperties(Properties);
+		}
+	}
+	else
+	{
+		//SplineInstanceData->SplineCurvesPreUCS = SplineCurves;
+	}
+
+	if (SplineInstanceData->bSplineHasBeenEdited)
+	{
+		// Copy metadata to current component
+		if (CineSplineMetadata && SplineInstanceData->CineSplineMetadata)
+		{
+			CineSplineMetadata->Modify();
+			UEngine::CopyPropertiesForUnrelatedObjects(SplineInstanceData->CineSplineMetadata, CineSplineMetadata);
+		}
+	
+		bModifiedByConstructionScript = false;
+	}
+
+	UpdateSpline();
+	SynchronizeProperties();
+}
+
+void UCineSplineComponent::SynchronizeProperties()
+{
+	int32 NumOfPoints = GetNumberOfSplinePoints();
+	if (CineSplineMetadata && NumOfPoints > 0)
+	{
+		CineSplineMetadata->Fixup(NumOfPoints, this);
+
+		// Fixing invalid CustomPosition
+		// This is to make sure CustomPostion value is incrementing from previous point
+		float PreviousValue = CineSplineMetadata->CustomPosition.Points[0].OutVal;
+		for (int32 Index = 1; Index < NumOfPoints; ++Index)
+		{
+			float CurrentValue = CineSplineMetadata->CustomPosition.Points[Index].OutVal;
+			if (CurrentValue <= PreviousValue)
+			{
+				CurrentValue = PreviousValue + 1.0f;
+				CineSplineMetadata->CustomPosition.Points[Index].OutVal = CurrentValue;
+			}
+			PreviousValue = CurrentValue;
+		}
+	}
+}
+
+void UCineSplineComponent::SetFocalLengthAtSplinePoint(const int32 PointIndex, const float Value)
+{
+	int32 NumPoints = CineSplineMetadata-> FocalLength.Points.Num();
+	check(PointIndex >= 0 && PointIndex < NumPoints);
+	CineSplineMetadata->Modify();
+	CineSplineMetadata->FocalLength.Points[PointIndex].OutVal = Value;
+}
+
+void UCineSplineComponent::SetApertureAtSplinePoint(const int32 PointIndex, const float Value)
+{
+	int32 NumPoints = CineSplineMetadata->Aperture.Points.Num();
+	check(PointIndex >= 0 && PointIndex < NumPoints);
+	CineSplineMetadata->Modify();
+	CineSplineMetadata->Aperture.Points[PointIndex].OutVal = Value;
+}
+
+void UCineSplineComponent::SetFocusDistanceAtSplinePoint(const int32 PointIndex, const float Value)
+{
+	int32 NumPoints = CineSplineMetadata->FocusDistance.Points.Num();
+	check(PointIndex >= 0 && PointIndex < NumPoints);
+	CineSplineMetadata->Modify();
+	CineSplineMetadata->FocusDistance.Points[PointIndex].OutVal = Value;
+}
+
+
+void UCineSplineComponent::SetCustomPositionAtSplinePoint(const int32 PointIndex, const float Value)
+{
+	int32 NumPoints = CineSplineMetadata->CustomPosition.Points.Num();
+	check(PointIndex >= 0 && PointIndex < NumPoints);
+	CineSplineMetadata->Modify();
+	CineSplineMetadata->CustomPosition.Points[PointIndex].OutVal = Value;
+}
+
+void UCineSplineComponent::SetCameraRotationAtSplinePoint(const int32 PointIndex, const FQuat Value)
+{
+	int32 NumPoints = CineSplineMetadata->CameraRotation.Points.Num();
+	check(PointIndex >= 0 && PointIndex < NumPoints);
+	CineSplineMetadata->Modify();
+	CineSplineMetadata->CameraRotation.Points[PointIndex].OutVal = Value;
+}
+
+bool UCineSplineComponent::FindSplineDataAtPosition(const float InPosition, int32& OutIndex, const float Tolerance) const
+{
+	int32 NumPoints = CineSplineMetadata->CustomPosition.Points.Num();
+	for (int32 i = 0; i < NumPoints; ++i)
+	{
+		if (FMath::IsNearlyEqual(InPosition, CineSplineMetadata->CustomPosition.Points[i].OutVal, Tolerance))
+		{
+			OutIndex = i;
+			return true;
+		}
+	}
+	OutIndex = -1;
+	return false;
+}
+
+float UCineSplineComponent::GetInputKeyAtPosition(const float InPosition) const
+{
+	// TODO: handle closed loop
+
+	float OutValue = 0.0f;
+	
+	int32 NumPoints = CineSplineMetadata->CustomPosition.Points.Num();
+	for (int32 i = 0; i < NumPoints; ++i)
+	{
+		if (InPosition < CineSplineMetadata->CustomPosition.Points[i].OutVal)
+		{
+			if (i > 0)
+			{
+				float Value0 = CineSplineMetadata->CustomPosition.Points[i - 1].OutVal;
+				float Value1 = CineSplineMetadata->CustomPosition.Points[i].OutVal;
+				OutValue = (InPosition - Value0) / (Value1 - Value0) + (float)(i-1);
+			}
+			break;
+		}
+		OutValue = (float)i;
+	}
+
+	return OutValue;
+}
+
+float UCineSplineComponent::GetPositionAtInputKey(const float InKey) const
+{
+	return GetFloatPropertyAtSplineInputKey(InKey, FName(TEXT("CustomPosition")));
+}
+
+FQuat UCineSplineComponent::GetCameraRotationAtSplinePoint(int32 Index) const
+{
+	const int32 NumPoints = CineSplineMetadata->CameraRotation.Points.Num();
+	if (NumPoints > 0)
+	{
+		const int32 ClampedIndex = FMath::Clamp(Index, 0, NumPoints - 1);
+		return CineSplineMetadata->CameraRotation.Points[ClampedIndex].OutVal;
+	}
+	
+	return FQuat::Identity;
+}
+
+FQuat UCineSplineComponent::GetCameraRotationAtSplineInputKey(float InKey) const
+{
+	return CineSplineMetadata->CameraRotation.Eval(InKey, FQuat::Identity);
+}
+
+void UCineSplineComponent::UpdateSplineDataAtIndex(const int InIndex, const FCineSplinePointData& InPointData)
+{
+	ESplinePointType::Type PointType = GetSplinePointType(InIndex);
+	SetLocationAtSplinePoint(InIndex, InPointData.Location, ESplineCoordinateSpace::World, false);
+	SetRotationAtSplinePoint(InIndex, InPointData.Rotation, ESplineCoordinateSpace::World, false);
+	SetSplinePointType(InIndex, PointType, false);
+	SetFocalLengthAtSplinePoint(InIndex, InPointData.FocalLength);
+	SetApertureAtSplinePoint(InIndex, InPointData.Aperture);
+	SetFocusDistanceAtSplinePoint(InIndex, InPointData.FocusDistance);
+	const FQuat Quat = GetComponentTransform().InverseTransformRotation(InPointData.Rotation.Quaternion());
+	SetCameraRotationAtSplinePoint(InIndex, Quat);
+	UpdateSpline();
+}
+
+void UCineSplineComponent::AddSplineDataAtPosition(const float InPosition, const FCineSplinePointData& InPointData)
+{
+	int32 NewIndex = 0;
+	int32 NumPoints = CineSplineMetadata->CustomPosition.Points.Num();
+	for (int32 i = 0; i < NumPoints; ++i)
+	{
+		if (InPosition <= CineSplineMetadata->CustomPosition.Points[i].OutVal)
+		{
+			break;
+		}
+		NewIndex++;
+	}
+	AddSplinePointAtIndex(InPointData.Location, NewIndex, ESplineCoordinateSpace::World, false);
+	SetRotationAtSplinePoint(NewIndex, InPointData.Rotation, ESplineCoordinateSpace::World, false);
+	SetFocalLengthAtSplinePoint(NewIndex, InPointData.FocalLength);
+	SetApertureAtSplinePoint(NewIndex, InPointData.Aperture);
+	SetFocusDistanceAtSplinePoint(NewIndex, InPointData.FocusDistance);
+	SetCustomPositionAtSplinePoint(NewIndex, InPosition);
+	SetSplinePointType(NewIndex, ESplinePointType::Curve, false);
+	const FQuat Quat = GetComponentTransform().InverseTransformRotation(InPointData.Rotation.Quaternion());
+	SetCameraRotationAtSplinePoint(NewIndex, Quat);
+	UpdateSpline();
+}
+
+FCineSplinePointData UCineSplineComponent::GetSplineDataAtPosition(const float InPosition) const
+{
+	FCineSplinePointData PointData;
+	const float InputKey = GetInputKeyAtPosition(InPosition);
+	PointData.Location = GetLocationAtSplineInputKey(InputKey, ESplineCoordinateSpace::World);
+	PointData.Rotation = GetCameraRotationAtSplineInputKey(InputKey).Rotator();
+	PointData.FocalLength = GetFloatPropertyAtSplineInputKey(InputKey, FName(TEXT("FocalLength")));
+	PointData.Aperture = GetFloatPropertyAtSplineInputKey(InputKey, FName(TEXT("Aperture")));
+	PointData.FocusDistance = GetFloatPropertyAtSplineInputKey(InputKey, FName(TEXT("FocusDistance")));
+	return PointData;
+}
+
+void UCineSplineComponent::UpdateSpline()
+{
+	USplineComponent::UpdateSpline();
+	if (OnSplineEdited.IsBound())
+	{
+		OnSplineEdited.Broadcast();
+	}
+	if (OnSplineEdited_BP.IsBound())
+	{
+		OnSplineEdited_BP.Broadcast();
+	}
+}
+
+
+#if WITH_EDITOR
+void UCineSplineComponent::PostEditChangeProperty(FPropertyChangedEvent& PropertyChangedEvent)
+{
+	if (OnSplineEdited.IsBound())
+	{
+		OnSplineEdited.Broadcast();
+	}
+	if (OnSplineEdited_BP.IsBound())
+	{
+		OnSplineEdited_BP.Broadcast();
+	}
+
+	Super::PostEditChangeProperty(PropertyChangedEvent);
+	SynchronizeProperties();
+}
+#endif
+
+void FCineSplineInstanceData::ApplyToComponent(UActorComponent* Component, const ECacheApplyPhase CacheApplyPhase)
+{
+	if (UCineSplineComponent* SplineComp = CastChecked<UCineSplineComponent>(Component))
+	{
+		// This ensures there is no stale data causing issues where the spline is marked as read-only even though it shouldn't.
+		// There might be a better solution, but this works.
+		SplineComp->UpdateSpline();
+
+		Super::ApplyToComponent(Component, CacheApplyPhase);
+		SplineComp->ApplyComponentInstanceData(this, (CacheApplyPhase == ECacheApplyPhase::PostUserConstructionScript));
+	}
+}
+
+
diff --git a/Engine/Plugins/Experimental/CineCameraRigs/Source/CineCameraRigs/Private/CineSplineLog.h b/Engine/Plugins/Experimental/CineCameraRigs/Source/CineCameraRigs/Private/CineSplineLog.h
new file mode 100644
--- /dev/null
+++ b/Engine/Plugins/Experimental/CineCameraRigs/Source/CineCameraRigs/Private/CineSplineLog.h
@@ -0,0 +1,8 @@
+// Copyright Epic Games, Inc. All Rights Reserved.
+
+#pragma once
+
+#include "CoreMinimal.h"
+#include "Logging/LogMacros.h"
+
+DECLARE_LOG_CATEGORY_EXTERN(LogCineSpline, Log, All);
diff --git a/Engine/Plugins/Experimental/CineCameraRigs/Source/CineCameraRigs/Private/CineSplineMetadata.cpp b/Engine/Plugins/Experimental/CineCameraRigs/Source/CineCameraRigs/Private/CineSplineMetadata.cpp
new file mode 100644
--- /dev/null
+++ b/Engine/Plugins/Experimental/CineCameraRigs/Source/CineCameraRigs/Private/CineSplineMetadata.cpp
@@ -0,0 +1,256 @@
+// Copyright Epic Games, Inc. All Rights Reserved.
+
+#include "CineSplineMetadata.h"
+#include "CineSplineComponent.h"
+#include "CineSplineLog.h"
+
+void UCineSplineMetadata::InsertPoint(int32 Index, float Alpha, bool bClosedLoop)
+{
+	if (Index < 0)
+	{
+		UE_LOG(LogCineSpline, Error, TEXT("InsertPoint - invalid Index: %d"), Index);
+		return;
+	}
+	Alpha = FMath::Clamp(Alpha, 0.f, 1.f);
+
+	Modify();
+
+	int32 NumPoints = FocalLength.Points.Num();
+	float InputKey = static_cast<float>(Index);
+	if (Index >= NumPoints)
+	{
+		// Just add point to the end instead of trying to insert
+		AddPoint(InputKey);
+	}
+	else
+	{
+		const int32 PrevIndex = (bClosedLoop && Index == 0 ? NumPoints - 1 : Index - 1);
+		const bool bHasPrevIndex = (PrevIndex >= 0 && PrevIndex < NumPoints);
+		float NewFocalLengthVal = FocalLength.Points[Index].OutVal;
+		float NewApertureVal = Aperture.Points[Index].OutVal;
+		float NewFocusDistanceVal = FocusDistance.Points[Index].OutVal;
+		float NewCustomPositionVal = CustomPosition.Points[Index].OutVal;
+		FQuat NewRotationVal = CameraRotation.Points[Index].OutVal;
+		if (bHasPrevIndex)
+		{
+			float PrevFocalLengthVal = FocalLength.Points[PrevIndex].OutVal;
+			float PrevApertureVal = Aperture.Points[PrevIndex].OutVal;
+			float PrevFocusDistanceVal = FocusDistance.Points[PrevIndex].OutVal;
+			float PrevCustomPositionVal = CustomPosition.Points[PrevIndex].OutVal;
+			FQuat PrevRotationVal = CameraRotation.Points[PrevIndex].OutVal;
+			NewFocalLengthVal = FMath::LerpStable(PrevFocalLengthVal, NewFocalLengthVal, Alpha);
+			NewApertureVal = FMath::LerpStable(PrevApertureVal, NewApertureVal, Alpha);
+			NewFocusDistanceVal = FMath::LerpStable(PrevFocusDistanceVal, NewFocusDistanceVal, Alpha);
+			NewCustomPositionVal = FMath::LerpStable(PrevCustomPositionVal, NewCustomPositionVal, Alpha);
+			NewRotationVal = FMath::LerpStable(PrevRotationVal, NewRotationVal, Alpha);
+		}
+		FInterpCurvePoint<float> NewCustomPosition(InputKey, NewCustomPositionVal);
+		CustomPosition.Points.Insert(NewCustomPosition, Index);
+
+		FInterpCurvePoint<float> NewFocalLength(InputKey, NewFocalLengthVal);
+		FocalLength.Points.Insert(NewFocalLength, Index);
+
+		FInterpCurvePoint<float> NewAperture(InputKey, NewApertureVal);
+		Aperture.Points.Insert(NewAperture, Index);
+
+		FInterpCurvePoint<float> NewFocusDistance(InputKey, NewFocusDistanceVal);
+		FocusDistance.Points.Insert(NewFocusDistance, Index);
+
+		FInterpCurvePoint<FQuat> NewRotation(InputKey, NewRotationVal);
+		CameraRotation.Points.Insert(NewRotation, Index);
+
+
+		for (int32 i = Index + 1; i < FocalLength.Points.Num(); ++i)
+		{
+			FocalLength.Points[i].InVal += 1.0f;
+			Aperture.Points[i].InVal += 1.0f;
+			FocusDistance.Points[i].InVal += 1.0f;
+			CustomPosition.Points[i].InVal += 1.0f;
+			CameraRotation.Points[i].InVal += 1.0f;
+		}
+	}
+}
+
+void UCineSplineMetadata::UpdatePoint(int32 Index, float Alpha, bool bClosedLoop)
+{
+	int32 NumPoints = FocalLength.Points.Num();
+	if (!FocalLength.Points.IsValidIndex(Index))
+	{
+		UE_LOG(LogCineSpline, Error, TEXT("UpdatePoint - invalid Index: %d"), Index);
+		return;
+	}
+	Alpha = FMath::Clamp(Alpha, 0.f, 1.f);
+
+	int32 PrevIndex = (bClosedLoop && Index == 0 ? NumPoints - 1 : Index - 1);
+	int32 NextIndex = (bClosedLoop && Index + 1 > NumPoints ? 0 : Index + 1);
+
+	bool bHasPrevIndex = (PrevIndex >= 0 && PrevIndex < NumPoints);
+	bool bHasNextIndex = (NextIndex >= 0 && NextIndex < NumPoints);
+
+	Modify();
+
+	if (bHasPrevIndex && bHasNextIndex)
+	{
+		float PrevCustomPositionVal = CustomPosition.Points[PrevIndex].OutVal;
+		float PrevFocalLengthVal = FocalLength.Points[PrevIndex].OutVal;
+		float PrevApertureVal = Aperture.Points[PrevIndex].OutVal;
+		float PrevFocusDistanceVal = FocusDistance.Points[PrevIndex].OutVal;
+		FQuat PrevRotationVal = CameraRotation.Points[PrevIndex].OutVal;
+		float NextCustomPositionVal = CustomPosition.Points[NextIndex].OutVal;
+		float NextFocalLengthVal = FocalLength.Points[NextIndex].OutVal;
+		float NextApertureVal = Aperture.Points[NextIndex].OutVal;
+		float NextFocusDistanceVal = FocusDistance.Points[NextIndex].OutVal;
+		FQuat NextRotationVal = CameraRotation.Points[NextIndex].OutVal;
+
+		CustomPosition.Points[Index].OutVal = FMath::LerpStable(PrevCustomPositionVal, NextCustomPositionVal, Alpha);
+		FocalLength.Points[Index].OutVal = FMath::LerpStable(PrevFocalLengthVal, NextFocalLengthVal, Alpha);
+		Aperture.Points[Index].OutVal = FMath::LerpStable(PrevApertureVal, NextApertureVal, Alpha);
+		FocusDistance.Points[Index].OutVal = FMath::LerpStable(PrevFocusDistanceVal, NextFocusDistanceVal, Alpha);
+		CameraRotation.Points[Index].OutVal = FMath::LerpStable(PrevRotationVal, NextRotationVal, Alpha);
+	}
+}
+
+void UCineSplineMetadata::AddPoint(float InputKey)
+{
+	Modify();
+
+	float NewCustomPositionVal = -1.0f;
+	float NewFocalLengthVal = 35.0f;
+	float NewApertureVal = 2.8f;
+	float NewFocusDistanceVal = 100000.f;
+	FQuat NewRotationVal = FQuat::Identity;
+	
+	int Index = FocalLength.Points.Num() - 1;
+	if (Index >= 0)
+	{
+		NewFocalLengthVal = FocalLength.Points[Index].OutVal;
+		NewApertureVal = Aperture.Points[Index].OutVal;
+		NewFocusDistanceVal = FocusDistance.Points[Index].OutVal;
+		NewCustomPositionVal = CustomPosition.Points[Index].OutVal + 1.0f;
+		NewRotationVal = CameraRotation.Points[Index].OutVal;
+	}
+
+	float NewInputKey = static_cast<float>(++Index);
+	FocalLength.Points.Emplace(NewInputKey, NewFocalLengthVal);
+	Aperture.Points.Emplace(NewInputKey, NewApertureVal);
+	FocusDistance.Points.Emplace(NewInputKey, NewFocusDistanceVal);
+	CustomPosition.Points.Emplace(NewInputKey, NewCustomPositionVal);
+	CameraRotation.Points.Emplace(NewInputKey, NewRotationVal);
+}
+
+void UCineSplineMetadata::RemovePoint(int32 Index)
+{
+	check(Index < FocalLength.Points.Num());
+
+	Modify();
+	CustomPosition.Points.RemoveAt(Index);
+	FocalLength.Points.RemoveAt(Index);
+	Aperture.Points.RemoveAt(Index);
+	FocusDistance.Points.RemoveAt(Index);
+	CameraRotation.Points.RemoveAt(Index);
+
+	for (int32 i = Index; i < FocalLength.Points.Num(); ++i)
+	{
+		CustomPosition.Points[i].InVal -= 1.0f;
+		FocalLength.Points[i].InVal -= 1.0f;
+		Aperture.Points[i].InVal -= 1.0f;
+		FocusDistance.Points[i].InVal -= 1.0f;
+		CameraRotation.Points[i].InVal -= 1.0f;
+	}
+}
+
+void UCineSplineMetadata::DuplicatePoint(int32 Index)
+{
+	check(Index < FocalLength.Points.Num());
+
+	int32 NumPoints = FocalLength.Points.Num();
+	float NewValue = CustomPosition.Points[Index].OutVal + 1.0f;
+	if (Index < NumPoints - 1)
+	{
+		NewValue = (CustomPosition.Points[Index].OutVal + CustomPosition.Points[Index + 1].OutVal) * 0.5;
+	}
+
+	Modify();
+	CustomPosition.Points.Insert(FInterpCurvePoint<float>(CustomPosition.Points[Index]), Index);
+	FocalLength.Points.Insert(FInterpCurvePoint<float>(FocalLength.Points[Index]), Index);
+	Aperture.Points.Insert(FInterpCurvePoint<float>(Aperture.Points[Index]), Index);
+	FocusDistance.Points.Insert(FInterpCurvePoint<float>(FocusDistance.Points[Index]), Index);
+	CameraRotation.Points.Insert(FInterpCurvePoint<FQuat>(CameraRotation.Points[Index]), Index);
+	CustomPosition.Points[Index+1].OutVal = NewValue;
+
+
+	for (int32 i = Index + 1; i < FocalLength.Points.Num(); ++i)
+	{
+		CustomPosition.Points[i].InVal += 1.0f;
+		FocalLength.Points[i].InVal += 1.0f;
+		Aperture.Points[i].InVal += 1.0f;
+		FocusDistance.Points[i].InVal += 1.0f;
+		CameraRotation.Points[i].InVal += 1.0f;
+	}
+
+}
+
+void UCineSplineMetadata::CopyPoint(const USplineMetadata* FromSplineMetadata, int32 FromIndex, int32 ToIndex)
+{
+	check(FromSplineMetadata != nullptr);
+
+	if (const UCineSplineMetadata* FromMetadata = Cast<UCineSplineMetadata>(FromSplineMetadata))
+	{
+		check(ToIndex < FocalLength.Points.Num());
+		check(FromIndex < FromMetadata->FocalLength.Points.Num());
+
+		Modify();
+		FocalLength.Points[ToIndex].OutVal = FromMetadata->FocalLength.Points[FromIndex].OutVal;
+		Aperture.Points[ToIndex].OutVal = FromMetadata->Aperture.Points[FromIndex].OutVal;
+		FocusDistance.Points[ToIndex].OutVal = FromMetadata->FocusDistance.Points[FromIndex].OutVal;
+		CustomPosition.Points[ToIndex].OutVal = FromMetadata->CustomPosition.Points[FromIndex].OutVal;
+		CameraRotation.Points[ToIndex].OutVal = FromMetadata->CameraRotation.Points[FromIndex].OutVal;
+	}
+}
+
+void UCineSplineMetadata::Reset(int32 NumPoints)
+{
+	Modify();
+	FocalLength.Points.Reset(NumPoints);
+	Aperture.Points.Reset(NumPoints);
+	FocusDistance.Points.Reset(NumPoints);
+	CustomPosition.Points.Reset(NumPoints);
+	CameraRotation.Points.Reset(NumPoints);
+}
+
+#if WITH_EDITORONLY_DATA
+template <class T>
+void FixupCurve(FInterpCurve<T>& Curve, const T& DefaultValue, int32 NumPoints)
+{
+	// Fixup bad InVal values from when the add operation below used the wrong value
+	for (int32 PointIndex = 0; PointIndex < Curve.Points.Num(); PointIndex++)
+	{
+		float InVal = PointIndex;
+		Curve.Points[PointIndex].InVal = InVal;
+	}
+
+	while (Curve.Points.Num() < NumPoints)
+	{
+		// InVal is the point index which is ascending so use previous point plus one.
+		float InVal = Curve.Points.Num() > 0 ? Curve.Points[Curve.Points.Num() - 1].InVal + 1.0f : 0.0f;
+		Curve.Points.Add(FInterpCurvePoint<T>(InVal, DefaultValue));
+	}
+
+	if (Curve.Points.Num() > NumPoints)
+	{
+		Curve.Points.RemoveAt(NumPoints, Curve.Points.Num() - NumPoints);
+	}
+}
+#endif
+
+void UCineSplineMetadata::Fixup(int32 NumPoints, USplineComponent* SplineComp)
+{
+	const FCineSplineCurveDefaults& Defaults = CastChecked<UCineSplineComponent>(SplineComp)->CameraSplineDefaults;
+#if WITH_EDITORONLY_DATA
+	FixupCurve(FocalLength, Defaults.DefaultFocalLength, NumPoints);
+	FixupCurve(Aperture, Defaults.DefaultAperture, NumPoints);
+	FixupCurve(FocusDistance, Defaults.DefaultFocusDistance, NumPoints);
+	FixupCurve(CustomPosition, Defaults.DefaultCustomPosition, NumPoints);
+	FixupCurve(CameraRotation, Defaults.DefaultCameraRotation, NumPoints);
+#endif
+}
\ No newline at end of file
diff --git a/Engine/Plugins/Experimental/CineCameraRigs/Source/CineCameraRigs/Public/CineCameraRigRail.h b/Engine/Plugins/Experimental/CineCameraRigs/Source/CineCameraRigs/Public/CineCameraRigRail.h
new file mode 100644
--- /dev/null
+++ b/Engine/Plugins/Experimental/CineCameraRigs/Source/CineCameraRigs/Public/CineCameraRigRail.h
@@ -0,0 +1,109 @@
+// Copyright Epic Games, Inc. All Rights Reserved.
+#pragma once
+
+#include "CoreMinimal.h"
+#include "UObject/ObjectMacros.h"
+
+#include "CameraRig_Rail.h"
+#include "CineSplineComponent.h"
+
+#include "CineCameraRigRail.generated.h"
+
+UCLASS(Blueprintable, Category = "VirtualProduction")
+class CINECAMERARIGS_API ACineCameraRigRail : public ACameraRig_Rail
+{
+	GENERATED_BODY()
+
+public:
+	ACineCameraRigRail(const FObjectInitializer& ObjectInitializer);
+
+	/* Returns CineSplineComponent*/
+	UFUNCTION(BlueprintPure, Category = "Rail Components")
+	UCineSplineComponent* GetCineSplineComponent() const { return CineSplineComponent; }
+
+	/* Use CustomPosition metadata to parameterize the spline*/
+	UPROPERTY(EditAnywhere, BlueprintReadWrite, Category = "Extended Rail Controls")
+	bool bUseCustomPosition = true;
+
+	/* Custom parameter to drive current position*/
+	UPROPERTY(EditAnywhere, BlueprintReadWrite, Interp, Category = "Extended Rail Controls", meta=(EditCondition="bUseCustomPosition"))
+	float CustomPosition = 1.0f;
+
+	/* Use CameraRotation metadata for attachment orientation*/
+	UPROPERTY(EditAnywhere, BlueprintReadWrite, Interp, Category = "Extended Rail Controls", meta = (EditCondition = "bLockOrientationToRail"))
+	bool bUseCameraRotationMetadata = true;
+
+	/* Material assigned to spline component mesh*/
+	UPROPERTY(EditAnywhere, BlueprintSetter=SetSplineMeshMaterial, Category = "SplineVisualization")
+	TObjectPtr<UMaterialInterface> SplineMeshMaterial;
+
+	/* Material Instance Dynamic created for the spline mesh */
+	UPROPERTY(VisibleInstanceOnly, Transient, BlueprintReadOnly, Category = "SplineVisualization")
+	TArray<TObjectPtr<UMaterialInstanceDynamic>> SplineMeshMIDs;
+
+	/* Texture that can be set to SplineMeshMIDs */
+	UPROPERTY(EditAnywhere, BlueprintSetter=SetSplineMeshTexture, Category = "SplineVisualization")
+	TObjectPtr<UTexture2D> SplineMeshTexture;
+
+	/* Determines if camera mount inherits LocationX*/
+	UPROPERTY(EditAnywhere, BlueprintReadWrite, Category = "Attachment")
+	bool bAttachLocationX = true;
+
+	/* Determines if camera mount inherits LocationY*/
+	UPROPERTY(EditAnywhere, BlueprintReadWrite, Category = "Attachment")
+	bool bAttachLocationY = true;
+
+	/* Determines if camera mount inherits LocationZ*/
+	UPROPERTY(EditAnywhere, BlueprintReadWrite, Category = "Attachment")
+	bool bAttachLocationZ = true;
+
+	/* Determines if camera mount inherits RotationX*/
+	UPROPERTY(EditAnywhere, BlueprintReadWrite, Category = "Attachment")
+	bool bAttachRotationX = true;
+
+	/* Determines if camera mount inherits RotationY*/
+	UPROPERTY(EditAnywhere, BlueprintReadWrite, Category = "Attachment")
+	bool bAttachRotationY = true;
+
+	/* Determines if camera mount inherits RotationZ*/
+	UPROPERTY(EditAnywhere, BlueprintReadWrite, Category = "Attachment")
+	bool bAttachRotationZ = true;
+
+	/* Determines if it can drive focal length on the attached actors*/
+	UPROPERTY(EditAnywhere, BlueprintReadWrite, Category = "Attachment")
+	bool bDriveFocalLength = true;
+
+	/* Determines if it can drive aperture on the attached actors*/
+	UPROPERTY(EditAnywhere, BlueprintReadWrite, Category = "Attachment")
+	bool bDriveAperture = true;
+
+	/* Determines if it can drive focus distance on the attached actors*/
+	UPROPERTY(EditAnywhere, BlueprintReadWrite, Category = "Attachment")
+	bool bDriveFocusDistance = true;
+
+	/* Set spline mesh material*/
+	UFUNCTION(BlueprintSetter)
+	void SetSplineMeshMaterial(UMaterialInterface* InMaterial);
+
+	/* Set texture used in the spline mesh material */
+	UFUNCTION(BlueprintSetter)
+	void SetSplineMeshTexture(UTexture2D* InTexture);
+
+	/* Calculate internal velocity at the given position */
+	UFUNCTION(BlueprintCallable, Category = "CineCameraRigRail")
+	FVector GetVelocityAtPosition(const float InPosition, const float delta = 0.001) const;
+
+#if WITH_EDITOR
+	virtual void PostEditChangeProperty(FPropertyChangedEvent& PropertyChangedEvent) override;
+#endif
+
+protected:
+	class UCineSplineComponent* CineSplineComponent;
+
+	virtual void UpdateRailComponents() override;
+
+	void UpdateSplineMeshMID();
+	void SetMIDParameters();
+
+	void OnSplineEdited();
+};
diff --git a/Engine/Plugins/Experimental/CineCameraRigs/Source/CineCameraRigs/Public/CineCameraRigRailHelpers.h b/Engine/Plugins/Experimental/CineCameraRigs/Source/CineCameraRigs/Public/CineCameraRigRailHelpers.h
new file mode 100644
--- /dev/null
+++ b/Engine/Plugins/Experimental/CineCameraRigs/Source/CineCameraRigs/Public/CineCameraRigRailHelpers.h
@@ -0,0 +1,19 @@
+// Copyright Epic Games, Inc. All Rights Reserved.
+
+#pragma once
+
+#include "CoreMinimal.h"
+#include "Kismet/BlueprintFunctionLibrary.h"
+#include "Engine/Texture2D.h"
+
+#include "CineCameraRigRailHelpers.generated.h"
+
+UCLASS()
+class CINECAMERARIGS_API UCineCameraRigRailHelpers : public UBlueprintFunctionLibrary
+{
+	GENERATED_BODY()
+public:
+	/* Create a transient heatmap texture from data values*/
+	UFUNCTION(BlueprintCallable, Category = "CineCameraRigRail")
+	static void CreateOrUpdateSplineHeatmapTexture(UPARAM(ref) UTexture2D*& Texture, const TArray<float>& DataValues, const float LowValue, const float AverageValue, const float HighValue);
+};
diff --git a/Engine/Plugins/Experimental/CineCameraRigs/Source/CineCameraRigs/Public/CineCameraRigs.h b/Engine/Plugins/Experimental/CineCameraRigs/Source/CineCameraRigs/Public/CineCameraRigs.h
new file mode 100644
--- /dev/null
+++ b/Engine/Plugins/Experimental/CineCameraRigs/Source/CineCameraRigs/Public/CineCameraRigs.h
@@ -0,0 +1,15 @@
+// Copyright Epic Games, Inc. All Rights Reserved.
+
+#pragma once
+
+#include "CoreMinimal.h"
+#include "Modules/ModuleManager.h"
+
+class FCineCameraRigsModule : public IModuleInterface
+{
+public:
+
+	/** IModuleInterface implementation */
+	virtual void StartupModule() override;
+	virtual void ShutdownModule() override;
+};
diff --git a/Engine/Plugins/Experimental/CineCameraRigs/Source/CineCameraRigs/Public/CineCameraRigsSettings.h b/Engine/Plugins/Experimental/CineCameraRigs/Source/CineCameraRigs/Public/CineCameraRigsSettings.h
new file mode 100644
--- /dev/null
+++ b/Engine/Plugins/Experimental/CineCameraRigs/Source/CineCameraRigs/Public/CineCameraRigsSettings.h
@@ -0,0 +1,20 @@
+// Copyright Epic Games, Inc. All Rights Reserved.
+
+#pragma once
+
+#include "UObject/Object.h"
+#include "UObject/SoftObjectPath.h"
+
+#include "CineCameraRigsSettings.generated.h"
+
+UCLASS(BlueprintType, Config = CineCameraRigs)
+class CINECAMERARIGS_API UCineCameraRigRailSettings : public UObject
+{
+	GENERATED_BODY()
+public:
+	UPROPERTY(config, EditAnywhere, BlueprintReadWrite, Category = CineCameraRigRail)
+	TSoftObjectPtr<UMaterialInterface> DefaultSplineMeshMaterial;
+
+	UPROPERTY(config, EditAnywhere, BlueprintReadWrite, Category = CineCameraRigRail)
+	TSoftObjectPtr<UTexture2D> DefaultSplineMeshTexture;
+};
\ No newline at end of file
diff --git a/Engine/Plugins/Experimental/CineCameraRigs/Source/CineCameraRigs/Public/CineSplineComponent.h b/Engine/Plugins/Experimental/CineCameraRigs/Source/CineCameraRigs/Public/CineSplineComponent.h
new file mode 100644
--- /dev/null
+++ b/Engine/Plugins/Experimental/CineCameraRigs/Source/CineCameraRigs/Public/CineSplineComponent.h
@@ -0,0 +1,133 @@
+// Copyright Epic Games, Inc. All Rights Reserved.
+
+#pragma once
+
+#include "CoreMinimal.h"
+#include "UObject/ObjectMacros.h"
+#include "Misc/CoreDelegates.h"
+
+#include "Components/SplineComponent.h"
+
+#include "CineSplineMetadata.h"
+#include "CineSplinePointData.h"
+
+#include "CineSplineComponent.generated.h"
+
+UCLASS(ClassGroup = Utility, ShowCategories = (Mobility), HideCategories = (Physics, Collision, Lighting, Rendering, Mobile), meta = (BlueprintSpawnableComponent))
+class CINECAMERARIGS_API UCineSplineComponent : public USplineComponent
+{
+	GENERATED_BODY()
+
+public:
+	UCineSplineComponent(const FObjectInitializer& ObjectInitializer);
+
+	/**
+	* Defaults which are used to propagate values to spline points on instances of this in the world
+	*/
+	UPROPERTY(Category = Camera, EditDefaultsOnly)
+	FCineSplineCurveDefaults CameraSplineDefaults;
+
+	// UObject interface
+#if WITH_EDITOR
+	virtual void PostEditChangeProperty(FPropertyChangedEvent& PropertyChangedEvent) override;
+#endif
+	virtual void PostLoad() override;
+	// End of UObject interface
+
+	/** Spline component interface */
+	virtual USplineMetadata* GetSplinePointsMetadata() override;
+	virtual const USplineMetadata* GetSplinePointsMetadata() const override;
+	virtual bool AllowsSplinePointScaleEditing() const override { return false; }
+
+	void ApplyComponentInstanceData(struct FCineSplineInstanceData* ComponentInstanceData, const bool bPostUCS);
+
+	// UActorComponent Interface
+	virtual TStructOnScope<FActorComponentInstanceData> GetComponentInstanceData() const override;
+
+	/** Pointer to metadata */
+	UPROPERTY(Instanced)
+	TObjectPtr<UCineSplineMetadata> CineSplineMetadata;
+
+	/** Set focal lenght metadata at a given splint point*/
+	UFUNCTION(BlueprintCallable, Category = "CineSpline")
+	void SetFocalLengthAtSplinePoint(const int32 PointIndex, const float Value );
+
+	/** Set aperture metadata at a given spline point*/
+	UFUNCTION(BlueprintCallable, Category = "CineSpline")
+	void SetApertureAtSplinePoint(const int32 PointIndex, const float value );
+
+	/** Set focus distance metadata at a given spline point*/
+	UFUNCTION(BlueprintCallable, Category = "CineSpline")
+	void SetFocusDistanceAtSplinePoint(const int32 PointIndex, const float value);
+
+	/** Set custom position metadata at a given spline point*/
+	UFUNCTION(BlueprintCallable, Category = "CineSpline")
+	void SetCustomPositionAtSplinePoint(const int32 PointIndex, const float value);
+
+	/* Set camera rotation metadata at a given spline point*/
+	UFUNCTION(BlueprintCallable, Category = "CineSpline")
+	void SetCameraRotationAtSplinePoint(const int32 PointIndex, const FQuat value);
+
+	/** Returns true if there is a spline point at given custom position*/
+	UFUNCTION(BlueprintCallable, BlueprintPure, Category = "CineSpline")
+	bool FindSplineDataAtPosition(const float InPosition, int32& OutIndex, const float Tolerance=0.001f) const;
+
+	/* Update spline point data at the given spline point*/
+	UFUNCTION(BlueprintCallable, Category = "CineSpline")
+	void UpdateSplineDataAtIndex(const int InIndex, const FCineSplinePointData& InPointData );
+
+	/* Add a new spline point data at the given custom position*/
+	UFUNCTION(BlueprintCallable, Category = "CineSpline")
+	void AddSplineDataAtPosition(const float InPosition, const FCineSplinePointData& InPointData);
+
+	/* Get a spline point data at the given custom position*/
+	UFUNCTION(BlueprintCallable, Category = "CineSpline")
+	FCineSplinePointData GetSplineDataAtPosition(const float InPosition) const;
+
+	/* Get input key value from custom position */
+	UFUNCTION(BlueprintCallable, Category = "CineSpline")
+	float GetInputKeyAtPosition(const float InPosition) const;
+
+	/* Get custom position value at spline input key */
+	UFUNCTION(BlueprintCallable, Category = "CineSpline")
+	float GetPositionAtInputKey(const float InKey) const;
+
+	/* Get camera rotation metadata property value along the spline at spline point */
+	UFUNCTION(BlueprintCallable, Category = "CineSpline")
+	FQuat GetCameraRotationAtSplinePoint(int32 Index) const;
+
+	/* Get camera rotation metadata value along the spline at spline input key*/
+	UFUNCTION(BlueprintCallable, Category = "CineSpline")
+	FQuat GetCameraRotationAtSplineInputKey(float InKey) const;
+
+
+	DECLARE_MULTICAST_DELEGATE(FOnSplineEdited);
+	DECLARE_DYNAMIC_MULTICAST_DELEGATE(FOnSplineEdited_BP);
+
+	FOnSplineEdited OnSplineEdited;
+
+	/* Event trigerred when spline is edited */
+	UPROPERTY(BlueprintAssignable, Category = "CineSpline", meta = (DisplayName = "OnSplineEdited"))
+	FOnSplineEdited_BP OnSplineEdited_BP;
+
+	/** Update the spline tangents and SplineReparamTable. And trigers OnSplineEdited delegates */
+	virtual void UpdateSpline() override;
+
+protected:
+	void SynchronizeProperties();
+
+};
+
+USTRUCT()
+struct FCineSplineInstanceData : public FSplineInstanceData
+{
+	GENERATED_BODY()
+public:
+	FCineSplineInstanceData() = default;
+	explicit FCineSplineInstanceData(const UCineSplineComponent* SourceComponent) : FSplineInstanceData(SourceComponent) {}
+
+	virtual void ApplyToComponent(UActorComponent* Component, const ECacheApplyPhase CacheApplyPhase) override;
+
+	UPROPERTY()
+	TObjectPtr<UCineSplineMetadata> CineSplineMetadata;
+};
\ No newline at end of file
diff --git a/Engine/Plugins/Experimental/CineCameraRigs/Source/CineCameraRigs/Public/CineSplineMetadata.h b/Engine/Plugins/Experimental/CineCameraRigs/Source/CineCameraRigs/Public/CineSplineMetadata.h
new file mode 100644
--- /dev/null
+++ b/Engine/Plugins/Experimental/CineCameraRigs/Source/CineCameraRigs/Public/CineSplineMetadata.h
@@ -0,0 +1,72 @@
+// Copyright Epic Games, Inc. All Rights Reserved.
+
+#pragma once
+
+#include "CoreMinimal.h"
+#include "UObject/ObjectMacros.h"
+
+#include "Components/SplineComponent.h"
+
+#include "CineSplineMetadata.generated.h"
+
+USTRUCT()
+struct FCineSplineCurveDefaults
+{
+	GENERATED_BODY()
+
+		FCineSplineCurveDefaults()
+		: DefaultCustomPosition(-1.0f)
+		, DefaultFocalLength(35.0f)
+		, DefaultAperture(2.8f)
+		, DefaultFocusDistance(100000.f)
+		, DefaultCameraRotation(FQuat::Identity)
+	{};
+
+	UPROPERTY(EditDefaultsOnly, Category = CineSplineCurveDefaults)
+	float DefaultCustomPosition;
+
+	UPROPERTY(EditDefaultsOnly, Category = CineSplineCurveDefaults)
+	float DefaultFocalLength;
+
+	UPROPERTY(EditDefaultsOnly, Category = CineSplineCurveDefaults)
+	float DefaultAperture;
+
+	UPROPERTY(EditDefaultsOnly, Category = CineSplineCurveDefaults)
+	float DefaultFocusDistance;
+
+	UPROPERTY(EditDefaultsOnly, Category = CineSplineCurveDefaults)
+	FQuat DefaultCameraRotation;
+};
+
+UCLASS()
+class CINECAMERARIGS_API UCineSplineMetadata : public USplineMetadata
+{
+	GENERATED_BODY()
+	
+public:
+	/** Insert point before index, lerping metadata between previous and next key values */
+	virtual void InsertPoint(int32 Index, float Alpha, bool bClosedLoop) override;
+	/** Update point at index by lerping metadata between previous and next key values */
+	virtual void UpdatePoint(int32 Index, float Alpha, bool bClosedLoop) override;
+	virtual void AddPoint(float InputKey) override;
+	virtual void RemovePoint(int32 Index) override;
+	virtual void DuplicatePoint(int32 Index) override;
+	virtual void CopyPoint(const USplineMetadata* FromSplineMetadata, int32 FromIndex, int32 ToIndex) override;
+	virtual void Reset(int32 NumPoints) override;
+	virtual void Fixup(int32 NumPoints, USplineComponent* SplineComp) override;
+
+	UPROPERTY(EditAnywhere, Category = "Point")
+	FInterpCurveFloat CustomPosition;
+
+	UPROPERTY(EditAnywhere, Category = "Camera")
+	FInterpCurveFloat FocalLength;
+
+	UPROPERTY(EditAnywhere, Category = "Camera")
+	FInterpCurveFloat Aperture;
+
+	UPROPERTY(EditAnywhere, Category = "Camera")
+	FInterpCurveFloat FocusDistance;
+
+	UPROPERTY(EditAnywhere, Category = "Camera")
+	FInterpCurveQuat CameraRotation;
+};
diff --git a/Engine/Plugins/Experimental/CineCameraRigs/Source/CineCameraRigs/Public/CineSplinePointData.h b/Engine/Plugins/Experimental/CineCameraRigs/Source/CineCameraRigs/Public/CineSplinePointData.h
new file mode 100644
--- /dev/null
+++ b/Engine/Plugins/Experimental/CineCameraRigs/Source/CineCameraRigs/Public/CineSplinePointData.h
@@ -0,0 +1,38 @@
+// Copyright Epic Games, Inc. All Rights Reserved.
+
+#pragma once
+
+#include "CoreMinimal.h"
+#include "UObject/ObjectMacros.h"
+
+#include "CineSplinePointData.generated.h"
+
+/* Simple struct to hold spline point data*/
+USTRUCT(BlueprintType)
+struct FCineSplinePointData
+{
+	GENERATED_BODY()
+
+	FCineSplinePointData()
+		: Location(FVector::ZeroVector)
+		, Rotation(FRotator::ZeroRotator)
+		,FocalLength(35.0f)
+		, Aperture(2.8f)
+		, FocusDistance(100000.f)
+	{};
+
+	UPROPERTY(BlueprintReadWrite, Category="CineSpline")
+	FVector Location;
+
+	UPROPERTY(BlueprintReadWrite, Category = "CineSpline")
+	FRotator Rotation;
+
+	UPROPERTY(BlueprintReadWrite, Category = "CineSpline")
+	float FocalLength;
+
+	UPROPERTY(BlueprintReadWrite, Category = "CineSpline")
+	float Aperture;
+
+	UPROPERTY(BlueprintReadWrite, Category = "CineSpline")
+	float FocusDistance;
+};
\ No newline at end of file
diff --git a/Engine/Plugins/Experimental/CineCameraRigs/Source/CineCameraRigsEditor/CineCameraRigsEditor.Build.cs b/Engine/Plugins/Experimental/CineCameraRigs/Source/CineCameraRigsEditor/CineCameraRigsEditor.Build.cs
new file mode 100644
--- /dev/null
+++ b/Engine/Plugins/Experimental/CineCameraRigs/Source/CineCameraRigsEditor/CineCameraRigsEditor.Build.cs
@@ -0,0 +1,38 @@
+// Copyright Epic Games, Inc. All Rights Reserved.
+
+using UnrealBuildTool;
+
+public class CineCameraRigsEditor : ModuleRules
+{
+	public CineCameraRigsEditor(ReadOnlyTargetRules Target) : base(Target)
+	{
+		PCHUsage = ModuleRules.PCHUsageMode.UseExplicitOrSharedPCHs;
+		
+        PrivateIncludePaths.Add("CineCameraRigs/Private");
+
+		PrivateDependencyModuleNames.AddRange(
+			new string[] {
+				"CinematicCamera",
+				"Core",
+				"CoreUObject",
+                "ComponentVisualizers",
+				"DetailCustomizations",
+				"Engine",
+				"SlateCore",
+				"Slate",
+				"UnrealEd",
+				"CineCameraRigs",
+				"PropertyEditor",
+				"InputCore",
+			});
+
+		PublicDependencyModuleNames.AddRange(
+			new string[] {
+			}
+		);
+
+		PrivateIncludePathModuleNames.AddRange(
+			new string[] {
+			});
+	}
+}
diff --git a/Engine/Plugins/Experimental/CineCameraRigs/Source/CineCameraRigsEditor/Private/CineCameraRigRailDetails.cpp b/Engine/Plugins/Experimental/CineCameraRigs/Source/CineCameraRigsEditor/Private/CineCameraRigRailDetails.cpp
new file mode 100644
--- /dev/null
+++ b/Engine/Plugins/Experimental/CineCameraRigs/Source/CineCameraRigsEditor/Private/CineCameraRigRailDetails.cpp
@@ -0,0 +1,172 @@
+// Copyright Epic Games, Inc. All Rights Reserved.
+
+#include "CineCameraRigRailDetails.h"
+#include "CineCameraRigRail.h"
+#include "CameraRig_Rail.h"
+#include "DetailCategoryBuilder.h"
+#include "DetailLayoutBuilder.h"
+#include "DetailWidgetRow.h"
+#include "Editor.h"
+#include "PropertyHandle.h"
+#include "UObject/PropertyIterator.h"
+#include "Widgets/Input/SNumericEntryBox.h"
+#include "Internationalization/Internationalization.h"
+#include "ScopedTransaction.h"
+
+#define LOCTEXT_NAMESPACE "FCineCameraRigRailDetails"
+
+void FCineCameraRigRailDetails::CustomizeDetails(IDetailLayoutBuilder& DetailBuilder)
+{
+	TArray<TWeakObjectPtr<UObject>> Objects;
+	DetailBuilder.GetObjectsBeingCustomized(Objects);
+
+	if (Objects.Num() != 1)
+	{
+		return;
+	}
+
+	ACineCameraRigRail* RigRailActor = Cast<ACineCameraRigRail>(Objects[0].Get());
+	if(RigRailActor == nullptr)
+	{
+		return;
+	}
+	RigRailActorPtr = RigRailActor;
+
+	TSharedRef<IPropertyHandle> CustomPositionPropertyHandle = DetailBuilder.GetProperty(GET_MEMBER_NAME_CHECKED(ACineCameraRigRail, CustomPosition));
+	TSharedRef<IPropertyHandle> UseCustomPositionPropertyHandle = DetailBuilder.GetProperty(GET_MEMBER_NAME_CHECKED(ACineCameraRigRail, bUseCustomPosition));
+	TSharedRef<IPropertyHandle> OrigPositionPropertyHandle = DetailBuilder.GetProperty(GET_MEMBER_NAME_CHECKED(ACameraRig_Rail, CurrentPositionOnRail), ACameraRig_Rail::StaticClass());
+
+	const TAttribute<bool> EditCondition = TAttribute<bool>::Create([this, UseCustomPositionPropertyHandle]()
+	{
+		bool bCond = false;
+		UseCustomPositionPropertyHandle->GetValue(bCond);
+		return !bCond;
+	});
+
+	DetailBuilder.EditDefaultProperty(OrigPositionPropertyHandle)->EditCondition(EditCondition, nullptr);
+
+
+	if (IDetailPropertyRow* Row = DetailBuilder.EditDefaultProperty(CustomPositionPropertyHandle))
+	{
+		Row->CustomWidget()
+		.NameContent()
+		[
+			CustomPositionPropertyHandle->CreatePropertyNameWidget()
+		]
+		.ValueContent()
+		.MinDesiredWidth(125.0f)
+		[
+			SNew(SNumericEntryBox<float>)
+			.ToolTipText(LOCTEXT("CurrentPostionToolTip","Postion property using custom parameterization"))
+			.AllowSpin(true)
+			.MinSliderValue(this, &FCineCameraRigRailDetails::GetCustomPositionSliderMinValue)
+			.MaxSliderValue(this, &FCineCameraRigRailDetails::GetCustomPositionSliderMaxValue)
+			.Value(this, &FCineCameraRigRailDetails::GetCustomPosition)
+			.OnValueChanged(this, &FCineCameraRigRailDetails::OnCustomPositionChanged)
+			.OnValueCommitted(this, &FCineCameraRigRailDetails::OnCustomPositionCommitted)
+			.OnBeginSliderMovement(this, &FCineCameraRigRailDetails::OnBeginCustomPositionSliderMovement)
+			.OnEndSliderMovement(this, &FCineCameraRigRailDetails::OnEndCustomPositionSliderMovement)
+			.Font(IDetailLayoutBuilder::GetDetailFont())
+		];
+	}
+
+}
+
+void FCineCameraRigRailDetails::OnCustomPositionChanged(float NewValue)
+{
+	if (ACineCameraRigRail* RigRailActor = RigRailActorPtr.Get())
+	{
+		RigRailActor->CustomPosition = NewValue;
+		static FProperty* CustomPositionProperty = FindFProperty<FProperty>(ACineCameraRigRail::StaticClass(), GET_MEMBER_NAME_CHECKED(ACineCameraRigRail, CustomPosition));
+		FPropertyChangedEvent SetValueEvent(CustomPositionProperty);
+		RigRailActor->PostEditChangeProperty(SetValueEvent);
+		if (GEditor)
+		{
+			GEditor->RedrawLevelEditingViewports(true);
+		}
+	}
+}
+
+void FCineCameraRigRailDetails::OnCustomPositionCommitted(float NewValue, ETextCommit::Type CommitType)
+{
+	if (ACineCameraRigRail* RigRailActor = RigRailActorPtr.Get())
+	{
+		const FScopedTransaction Transaction(LOCTEXT("SetCustomPosition", "Set rig rail custom position"));
+		RigRailActor->SetFlags(RF_Transactional);
+		RigRailActor->Modify();
+		RigRailActor->CustomPosition = NewValue;
+		static FProperty* CustomPositionProperty = FindFProperty<FProperty>(ACineCameraRigRail::StaticClass(), GET_MEMBER_NAME_CHECKED(ACineCameraRigRail, CustomPosition));
+		FPropertyChangedEvent SetValueEvent(CustomPositionProperty);
+		RigRailActor->PostEditChangeProperty(SetValueEvent);
+		if (GEditor)
+		{
+			GEditor->RedrawLevelEditingViewports(true);
+		}
+	}
+}
+
+void FCineCameraRigRailDetails::OnBeginCustomPositionSliderMovement()
+{
+	if (!bCustomPositionSliderStartedTransaction)
+	{
+		if (GEditor)
+		{
+			bCustomPositionSliderStartedTransaction = true;
+			GEditor->BeginTransaction(LOCTEXT("CustomPositionSliderTransaction", "Set rig rail custom position via slider"));
+			if (ACineCameraRigRail* RigRailActor = RigRailActorPtr.Get())
+			{
+				RigRailActor->SetFlags(RF_Transactional);
+				RigRailActor->Modify();
+			}
+		}
+	}
+}
+
+void FCineCameraRigRailDetails::OnEndCustomPositionSliderMovement(float NewValue)
+{
+	if (bCustomPositionSliderStartedTransaction)
+	{
+		if (GEditor)
+		{
+			GEditor->EndTransaction();
+			bCustomPositionSliderStartedTransaction = false;
+		}
+	}
+}
+
+TOptional<float> FCineCameraRigRailDetails::GetCustomPosition() const
+{
+	if (ACineCameraRigRail* RigRailActor = RigRailActorPtr.Get())
+	{
+		return RigRailActor->CustomPosition;
+	}
+	return 0.0f;
+}
+
+TOptional<float> FCineCameraRigRailDetails::GetCustomPositionSliderMinValue() const
+{
+
+	float MinValue = 1.0f;
+	if (ACineCameraRigRail* RigRailActor = RigRailActorPtr.Get())
+	{
+		UCineSplineComponent* SplineComp = RigRailActor->GetCineSplineComponent();
+		const UCineSplineMetadata* MetaData = Cast<UCineSplineMetadata>(SplineComp->GetSplinePointsMetadata());
+		MinValue = MetaData->CustomPosition.Points[0].OutVal;
+	}
+	return MinValue;
+}
+
+TOptional<float> FCineCameraRigRailDetails::GetCustomPositionSliderMaxValue() const
+{
+
+	float MaxValue = 5.0f;
+	if (ACineCameraRigRail* RigRailActor = RigRailActorPtr.Get())
+	{
+		UCineSplineComponent* SplineComp = RigRailActor->GetCineSplineComponent();
+		const UCineSplineMetadata* MetaData = Cast<UCineSplineMetadata>(SplineComp->GetSplinePointsMetadata());
+		int32 NumPoints = MetaData->CustomPosition.Points.Num();
+		MaxValue = MetaData->CustomPosition.Points[NumPoints - 1].OutVal;
+	}
+	return MaxValue;
+}
+#undef LOCTEXT_NAMESPACE
\ No newline at end of file
diff --git a/Engine/Plugins/Experimental/CineCameraRigs/Source/CineCameraRigsEditor/Private/CineCameraRigRailDetails.h b/Engine/Plugins/Experimental/CineCameraRigs/Source/CineCameraRigsEditor/Private/CineCameraRigRailDetails.h
new file mode 100644
--- /dev/null
+++ b/Engine/Plugins/Experimental/CineCameraRigs/Source/CineCameraRigsEditor/Private/CineCameraRigRailDetails.h
@@ -0,0 +1,39 @@
+// Copyright Epic Games, Inc. All Rights Reserved.
+
+#pragma once
+
+#include "CoreMinimal.h"
+#include "IDetailCustomization.h"
+#include "CineCameraRigRail.h"
+
+namespace ETextCommit { enum Type : int; }
+
+/**
+ * Implements a details view customization for the ACineCameraRigRail
+ */
+class FCineCameraRigRailDetails : public IDetailCustomization
+{
+public:
+	static TSharedRef<IDetailCustomization> MakeInstance()
+	{
+		return MakeShared<FCineCameraRigRailDetails>();
+	}
+
+	virtual void CustomizeDetails(IDetailLayoutBuilder& DetailBuilder) override;
+
+	TOptional<float> CurrentPositionValue = 1.0f;
+
+private:
+	TWeakObjectPtr<ACineCameraRigRail> RigRailActorPtr;
+
+	TOptional<float> GetCustomPosition() const;
+
+	TOptional<float> GetCustomPositionSliderMinValue() const;
+	TOptional<float> GetCustomPositionSliderMaxValue() const;
+	void OnBeginCustomPositionSliderMovement();
+	void OnEndCustomPositionSliderMovement(float NewValue);
+	void OnCustomPositionCommitted(float NewValue, ETextCommit::Type CommitType);
+	void OnCustomPositionChanged(float NewValue);
+
+	bool bCustomPositionSliderStartedTransaction = false;
+};
diff --git a/Engine/Plugins/Experimental/CineCameraRigs/Source/CineCameraRigsEditor/Private/CineCameraRigsEditor.cpp b/Engine/Plugins/Experimental/CineCameraRigs/Source/CineCameraRigsEditor/Private/CineCameraRigsEditor.cpp
new file mode 100644
--- /dev/null
+++ b/Engine/Plugins/Experimental/CineCameraRigs/Source/CineCameraRigsEditor/Private/CineCameraRigsEditor.cpp
@@ -0,0 +1,62 @@
+// Copyright Epic Games, Inc. All Rights Reserved.
+
+#include "CoreMinimal.h"
+#include "Modules/ModuleManager.h"
+#include "CineSplineComponentVisualizer.h"
+#include "CineSplineComponent.h"
+#include "CineCameraRigRail.h"
+#include "CineCameraRigRailDetails.h"
+#include "Editor.h"
+#include "Editor/UnrealEdEngine.h"
+#include "UnrealEdGlobals.h"
+#include "PropertyEditorModule.h"
+#include "PropertyEditorDelegates.h"
+
+#define LOCTEXT_NAMESPACE "CineCameraRigsEditor"
+
+class FCineCameraRigsEditorModule : public IModuleInterface
+{
+public:
+
+	virtual void StartupModule() override
+	{
+		CineSplineComponentName = UCineSplineComponent::StaticClass()->GetFName();
+
+		if (GUnrealEd)
+		{
+			TSharedPtr<FCineSplineComponentVisualizer> Visualizer = MakeShared<FCineSplineComponentVisualizer>();
+			GUnrealEd->RegisterComponentVisualizer(CineSplineComponentName, Visualizer);
+			Visualizer->OnRegister();
+			RegisterCustomizations();
+		}
+	}
+
+	virtual void ShutdownModule() override
+	{
+		if (GUnrealEd)
+		{
+			GUnrealEd->UnregisterComponentVisualizer(CineSplineComponentName);
+			UnregisterCustomizations();
+		}
+	}
+
+private:
+	FName CineSplineComponentName;
+
+
+	void RegisterCustomizations()
+	{
+		FPropertyEditorModule& PropertyModule = FModuleManager::LoadModuleChecked<FPropertyEditorModule>("PropertyEditor");
+		PropertyModule.RegisterCustomClassLayout(ACineCameraRigRail::StaticClass()->GetFName(), FOnGetDetailCustomizationInstance::CreateStatic(&FCineCameraRigRailDetails::MakeInstance));
+	}
+	void UnregisterCustomizations()
+	{
+		FPropertyEditorModule& PropertyModule = FModuleManager::GetModuleChecked<FPropertyEditorModule>("PropertyEditor");
+		PropertyModule.UnregisterCustomPropertyTypeLayout(ACineCameraRigRail::StaticClass()->GetFName());
+	}
+};
+
+
+#undef LOCTEXT_NAMESPACE
+
+IMPLEMENT_MODULE(FCineCameraRigsEditorModule, CineCameraRigsEditor)
diff --git a/Engine/Plugins/Experimental/CineCameraRigs/Source/CineCameraRigsEditor/Private/CineSplineComponentVisualizer.cpp b/Engine/Plugins/Experimental/CineCameraRigs/Source/CineCameraRigsEditor/Private/CineSplineComponentVisualizer.cpp
new file mode 100644
--- /dev/null
+++ b/Engine/Plugins/Experimental/CineCameraRigs/Source/CineCameraRigsEditor/Private/CineSplineComponentVisualizer.cpp
@@ -0,0 +1,56 @@
+// Copyright Epic Games, Inc. All Rights Reserved.
+
+#include "CineSplineComponentVisualizer.h"
+#include "Internationalization/Text.h"
+#include "Math/Color.h"
+#include "CanvasTypes.h"
+#include "CanvasItem.h"
+#include "SceneView.h"
+#include "Engine/Engine.h"
+
+#define LOCTEXT_NAMESPACE "SplineComponentVisualizer"
+
+void FCineSplineComponentVisualizer::OnRegister()
+{
+	FSplineComponentVisualizer::OnRegister();
+}
+
+void FCineSplineComponentVisualizer::DrawVisualizationHUD(const UActorComponent* Component, const FViewport* Viewport, const FSceneView* View, FCanvas* Canvas)
+{
+	FSplineComponentVisualizer::DrawVisualizationHUD(Component, Viewport, View, Canvas);
+	if (Canvas == nullptr || View == nullptr)
+	{
+		return;
+	}
+
+	const FIntRect CanvasRect = Canvas->GetViewRect();
+	const float HalfX = CanvasRect.Width() / 2.f;
+	const float HalfY = CanvasRect.Height() / 2.f;
+
+
+	if (const UCineSplineComponent* SplineComp = Cast<const UCineSplineComponent>(Component))
+	{
+		if (UCineSplineMetadata* Metadata = SplineComp->CineSplineMetadata)
+		{
+			const int32 NumOfPoints = SplineComp->GetNumberOfSplinePoints();
+			for (int32 i = 0; i < NumOfPoints; ++i)
+			{
+				FVector Location = SplineComp->GetLocationAtSplinePoint(i, ESplineCoordinateSpace::World);
+				FPlane Plane(0, 0, 0, 0);
+				Plane = View->Project(Location);
+				const FVector Position(Plane);
+				const float DrawPositionX = FMath::FloorToFloat(HalfX + Position.X * HalfX);
+				const float DrawPositionY = FMath::FloorToFloat(HalfY + -1.f * Position.Y * HalfY);
+				
+				FNumberFormattingOptions FmtOptions;
+				FmtOptions.SetMaximumFractionalDigits(4);
+				FmtOptions.SetMinimumFractionalDigits(1);
+				const FText Text = FText::AsNumber(Metadata->CustomPosition.Points[i].OutVal, &FmtOptions);
+				FLinearColor Color = FMath::IsNearlyEqual(FMath::Frac(Metadata->CustomPosition.Points[i].OutVal), 0.0f) ? FLinearColor(1,1,0,1) : FLinearColor(0, 1, 1, 1);
+				Canvas->DrawShadowedString(DrawPositionX, DrawPositionY, *Text.ToString(), GEngine->GetLargeFont(), Color, FLinearColor::Black);
+			}
+		}
+	}
+}
+
+#undef LOCTEXT_NAMESPACE
\ No newline at end of file
diff --git a/Engine/Plugins/Experimental/CineCameraRigs/Source/CineCameraRigsEditor/Private/CineSplineComponentVisualizer.h b/Engine/Plugins/Experimental/CineCameraRigs/Source/CineCameraRigsEditor/Private/CineSplineComponentVisualizer.h
new file mode 100644
--- /dev/null
+++ b/Engine/Plugins/Experimental/CineCameraRigs/Source/CineCameraRigsEditor/Private/CineSplineComponentVisualizer.h
@@ -0,0 +1,24 @@
+// Copyright Epic Games, Inc. All Rights Reserved.
+
+#pragma once
+
+#include "CoreMinimal.h"
+
+#include "HitProxies.h"
+#include "InputCoreTypes.h"
+
+#include "ComponentVisualizer.h"
+#include "SplineComponentVisualizer.h"
+
+#include "Components/SplineComponent.h"
+#include "CineSplineComponent.h"
+
+
+class FCineSplineComponentVisualizer : public FSplineComponentVisualizer
+{
+public:
+	//~ Begin FComponentVisualizer Interface
+	virtual void OnRegister() override;
+	virtual void DrawVisualizationHUD(const UActorComponent* Component, const FViewport* Viewport, const FSceneView* View, FCanvas* Canvas) override;
+
+};
\ No newline at end of file
diff --git a/Engine/Plugins/Experimental/CineCameraRigs/Source/CineCameraRigsEditor/Private/CineSplineMetadataDetails.cpp b/Engine/Plugins/Experimental/CineCameraRigs/Source/CineCameraRigsEditor/Private/CineSplineMetadataDetails.cpp
new file mode 100644
--- /dev/null
+++ b/Engine/Plugins/Experimental/CineCameraRigs/Source/CineCameraRigsEditor/Private/CineSplineMetadataDetails.cpp
@@ -0,0 +1,335 @@
+// Copyright Epic Games, Inc. All Rights Reserved.
+
+#include "CineSplineMetadataDetails.h"
+#include "CineSplineMetadata.h"
+
+#include "IDetailChildrenBuilder.h"
+#include "IDetailGroup.h"
+#include "DetailLayoutBuilder.h"
+#include "DetailWidgetRow.h"
+#include "Internationalization/Internationalization.h"
+#include "Widgets/Text/STextBlock.h"
+#include "Widgets/Input/SNumericEntryBox.h"
+#include "Widgets/Input/SRotatorInputBox.h"
+#include "ScopedTransaction.h"
+#include "ComponentVisualizer.h"
+#include "Widgets/Input/SVectorInputBox.h"
+#include "Editor.h"
+
+#define LOCTEXT_NAMESPACE "FCineSplineMetadataDetails"
+
+UClass* UCineSplineMetadataDetailsFactory::GetMetadataClass() const
+{
+	return UCineSplineMetadata::StaticClass();
+}
+
+TSharedPtr<ISplineMetadataDetails> UCineSplineMetadataDetailsFactory::Create()
+{
+	return MakeShared<FCineSplineMetadataDetails>();
+}
+
+FText FCineSplineMetadataDetails::GetDisplayName() const
+{
+	return LOCTEXT("CineSplineMetadataDetails", "CineSpline");
+}
+
+template<class T>
+bool UpdateMultipleValue(TOptional<T>& CurrentValue, T InValue)
+{
+	if (!CurrentValue.IsSet())
+	{
+		CurrentValue = InValue;
+	}
+	else if (CurrentValue.IsSet() && CurrentValue.GetValue() != InValue)
+	{
+		CurrentValue.Reset();
+		return false;
+	}
+
+	return true;
+}
+
+void FCineSplineMetadataDetails::Update(USplineComponent* InSplineComponent, const TSet<int32>& InSelectedKeys)
+{
+	SplineComp = InSplineComponent;
+	SelectedKeys = InSelectedKeys;
+	CustomPositionValue.Reset();
+	FocalLengthValue.Reset();
+	ApertureValue.Reset();
+	FocusDistanceValue.Reset();
+	CameraRotationValue.Reset();
+
+	if (!IsValid(InSplineComponent))
+	{
+		return;
+	}
+	UCineSplineMetadata* Metadata = Cast<UCineSplineMetadata>(InSplineComponent->GetSplinePointsMetadata());
+	if (!Metadata)
+	{
+		return;
+	}
+
+	bool bUpdateCustomPosition = true;
+	bool bUpdateFocalLength = true;
+	bool bUpdateAperture = true;
+	bool bUpdateFocusDistance = true;
+	bool bUpdateCameraRotation = true;
+
+	for (const int32 Index : InSelectedKeys)
+	{
+		bUpdateCustomPosition = bUpdateCustomPosition ? UpdateMultipleValue(CustomPositionValue, Metadata->CustomPosition.Points[Index].OutVal) : false;
+		bUpdateFocalLength = bUpdateFocalLength ? UpdateMultipleValue(FocalLengthValue, Metadata->FocalLength.Points[Index].OutVal) : false;
+		bUpdateAperture = bUpdateAperture ? UpdateMultipleValue(ApertureValue, Metadata->Aperture.Points[Index].OutVal) : false;
+		bUpdateFocusDistance = bUpdateFocusDistance ? UpdateMultipleValue(FocusDistanceValue, Metadata->FocusDistance.Points[Index].OutVal) : false;
+		bUpdateCameraRotation = bUpdateCameraRotation ? UpdateMultipleValue(CameraRotationValue, Metadata->CameraRotation.Points[Index].OutVal) : false;
+	}
+}
+
+template<class T>
+void SetValues(FCineSplineMetadataDetails& Details, TArray<FInterpCurvePoint<T>>& Points, const T& NewValue)
+{
+	Details.SplineComp->GetSplinePointsMetadata()->Modify();
+	for (int32 Index : Details.SelectedKeys)
+	{
+		Points[Index].OutVal = NewValue;
+	}
+
+	Details.SplineComp->UpdateSpline();
+	Details.SplineComp->bSplineHasBeenEdited = true;
+	static FProperty* SplineCurvesProperty = FindFProperty<FProperty>(USplineComponent::StaticClass(), GET_MEMBER_NAME_CHECKED(USplineComponent, SplineCurves));
+	FComponentVisualizer::NotifyPropertyModified(Details.SplineComp, SplineCurvesProperty);
+	Details.Update(Details.SplineComp, Details.SelectedKeys);
+
+	if (GEditor)
+	{
+		GEditor->RedrawLevelEditingViewports(true);
+	}
+}
+
+void FCineSplineMetadataDetails::OnBeginSliderMovement()
+{
+	EditSliderValueTransaction = MakeUnique<FScopedTransaction>(LOCTEXT("EditCineSplineProperty", "Edit CineSpline Property"));
+}
+
+void FCineSplineMetadataDetails::OnEndSliderMovement(float NewValue)
+{
+	EditSliderValueTransaction.Reset();
+}
+
+UCineSplineMetadata* FCineSplineMetadataDetails::GetMetadata() const
+{
+	return SplineComp ? Cast<UCineSplineMetadata>(SplineComp->GetSplinePointsMetadata()) : nullptr;
+}
+
+void FCineSplineMetadataDetails::OnSetCustomPosition(float NewValue, ETextCommit::Type CommitInfo)
+{
+	if (UCineSplineMetadata* Metadata = GetMetadata())
+	{
+		const FScopedTransaction Transaction(LOCTEXT("SetCustomPosition", "Set spline point custom position data"));
+		SetValues<float>(*this, Metadata->CustomPosition.Points, NewValue);
+	}
+}
+
+void FCineSplineMetadataDetails::OnSetFocalLength(float NewValue, ETextCommit::Type CommitInfo)
+{
+	if (UCineSplineMetadata* Metadata = GetMetadata())
+	{
+		const FScopedTransaction Transaction(LOCTEXT("SetFocalLength", "Set spline point focal length data"));
+		SetValues<float>(*this, Metadata->FocalLength.Points, NewValue);
+	}
+}
+
+void FCineSplineMetadataDetails::OnSetAperture(float NewValue, ETextCommit::Type CommitInfo)
+{
+	if (UCineSplineMetadata* Metadata = GetMetadata())
+	{
+		const FScopedTransaction Transaction(LOCTEXT("SetAperture", "Set spline point aperture data"));
+		SetValues<float>(*this, Metadata->Aperture.Points, NewValue);
+	}
+}
+
+void FCineSplineMetadataDetails::OnSetFocusDistance(float NewValue, ETextCommit::Type CommitInfo)
+{
+	if (UCineSplineMetadata* Metadata = GetMetadata())
+	{
+		const FScopedTransaction Transaction(LOCTEXT("SetFocusDistance", "Set spline point focus distance data"));
+		SetValues<float>(*this, Metadata->FocusDistance.Points, NewValue);
+	}
+}
+
+TOptional<float> FCineSplineMetadataDetails::GetRotation(EAxis::Type Axis) const
+{
+	return static_cast<float>(CameraRotationValue.Get(FQuat::Identity).Rotator().GetComponentForAxis(Axis));
+}
+void FCineSplineMetadataDetails::OnSetRotation(float NewValue, ETextCommit::Type CommitInfo, EAxis::Type Axis)
+{
+	if (UCineSplineMetadata* Metadata = GetMetadata())
+	{
+		const FScopedTransaction Transaction(LOCTEXT("SetCameraRotation", "Set spline point camera rotation data"));
+		FRotator CurrentRotator(CameraRotationValue.Get(FQuat::Identity));
+		switch (Axis)
+		{
+		case EAxis::X: CurrentRotator.Roll = NewValue; break;
+		case EAxis::Y: CurrentRotator.Pitch = NewValue; break;
+		case EAxis::Z: CurrentRotator.Yaw = NewValue; break;
+		}
+		SetValues<FQuat>(*this, Metadata->CameraRotation.Points, CurrentRotator.Quaternion());
+	}
+}
+
+void FCineSplineMetadataDetails::GenerateChildContent(IDetailGroup& DetailGroup)
+{
+	DetailGroup.AddWidgetRow()
+		.RowTag("CustomPosition")
+		.Visibility(TAttribute<EVisibility>(this, &FCineSplineMetadataDetails::IsEnabled))
+		.NameContent()
+		.HAlign(HAlign_Left)
+		.VAlign(VAlign_Center)
+		[
+			SNew(STextBlock)
+			.Text(LOCTEXT("CustomPosition", "Custom Position"))
+			.Font(IDetailLayoutBuilder::GetDetailFont())
+		]
+		.ValueContent()
+		.MinDesiredWidth(125.0f)
+		.MaxDesiredWidth(125.0f)
+		[
+			SNew(SNumericEntryBox<float>)
+			.Value(this, &FCineSplineMetadataDetails::GetCustomPosition)
+			.AllowSpin(true)
+			.MinValue(TOptional<float>())
+			.MaxValue(TOptional<float>())
+			.MinSliderValue(TOptional<float>())
+			.MaxSliderValue(TOptional<float>())
+			.OnBeginSliderMovement(this, &FCineSplineMetadataDetails::OnBeginSliderMovement)
+			.OnEndSliderMovement(this, &FCineSplineMetadataDetails::OnEndSliderMovement)
+			.UndeterminedString(LOCTEXT("Multiple", "Multiple"))
+			.OnValueCommitted(this, &FCineSplineMetadataDetails::OnSetCustomPosition)
+			.OnValueChanged(this, &FCineSplineMetadataDetails::OnSetCustomPosition, ETextCommit::Default)
+			.Font(IDetailLayoutBuilder::GetDetailFont())
+		];
+
+	DetailGroup.AddWidgetRow()
+		.RowTag("FocalLength")
+		.Visibility(TAttribute<EVisibility>(this, &FCineSplineMetadataDetails::IsEnabled))
+		.NameContent()
+		.HAlign(HAlign_Left)
+		.VAlign(VAlign_Center)
+		[
+			SNew(STextBlock)
+			.Text(LOCTEXT("FocalLength", "Focal Length"))
+			.Font(IDetailLayoutBuilder::GetDetailFont())
+		]
+		.ValueContent()
+		.MinDesiredWidth(125.0f)
+		.MaxDesiredWidth(125.0f)
+		[
+			SNew(SNumericEntryBox<float>)
+			.Value(this, &FCineSplineMetadataDetails::GetFocalLength)
+			.AllowSpin(true)
+			.OnBeginSliderMovement(this, &FCineSplineMetadataDetails::OnBeginSliderMovement)
+			.OnEndSliderMovement(this, &FCineSplineMetadataDetails::OnEndSliderMovement)
+			.MinValue(0.0f)
+			// Because we have no upper limit in MaxSliderValue, we need to "unspecify" the max value here, otherwise the spinner has a limited range,
+			//  with TNumericLimits<NumericType>::Max() as the MaxValue and the spinning increment is huge
+			.MaxValue(TOptional<float>())
+			.MinSliderValue(0.0f)
+			.MaxSliderValue(TOptional<float>()) // No upper limit
+			.UndeterminedString(LOCTEXT("Multiple", "Multiple"))
+			.OnValueCommitted(this, &FCineSplineMetadataDetails::OnSetFocalLength)
+			.OnValueChanged(this, &FCineSplineMetadataDetails::OnSetFocalLength, ETextCommit::Default)
+			.Font(IDetailLayoutBuilder::GetDetailFont())
+		];
+
+	DetailGroup.AddWidgetRow()
+		.RowTag("Aperture")
+		.Visibility(TAttribute<EVisibility>(this, &FCineSplineMetadataDetails::IsEnabled))
+		.NameContent()
+		.HAlign(HAlign_Left)
+		.VAlign(VAlign_Center)
+		[
+			SNew(STextBlock)
+			.Text(LOCTEXT("Aperture", "Aperture"))
+			.Font(IDetailLayoutBuilder::GetDetailFont())
+		]
+		.ValueContent()
+		.MinDesiredWidth(125.0f)
+		.MaxDesiredWidth(125.0f)
+		[
+			SNew(SNumericEntryBox<float>)
+			.Value(this, &FCineSplineMetadataDetails::GetAperture)
+			.AllowSpin(true)
+			.OnBeginSliderMovement(this, &FCineSplineMetadataDetails::OnBeginSliderMovement)
+			.OnEndSliderMovement(this, &FCineSplineMetadataDetails::OnEndSliderMovement)
+			.MinValue(TOptional<float>())
+			.MaxValue(TOptional<float>())
+			.MinSliderValue(TOptional<float>()) // No lower limit
+			.MaxSliderValue(TOptional<float>()) // No upper limit
+			.UndeterminedString(LOCTEXT("Multiple", "Multiple"))
+			.OnValueCommitted(this, &FCineSplineMetadataDetails::OnSetAperture)
+			.OnValueChanged(this, &FCineSplineMetadataDetails::OnSetAperture, ETextCommit::Default)
+			.Font(IDetailLayoutBuilder::GetDetailFont())
+		];
+
+	DetailGroup.AddWidgetRow()
+		.RowTag("FocusDistance")
+		.Visibility(TAttribute<EVisibility>(this, &FCineSplineMetadataDetails::IsEnabled))
+		.NameContent()
+		.HAlign(HAlign_Left)
+		.VAlign(VAlign_Center)
+		[
+			SNew(STextBlock)
+			.Text(LOCTEXT("FocusDistance", "Focus Distance"))
+			.Font(IDetailLayoutBuilder::GetDetailFont())
+		]
+		.ValueContent()
+		.MinDesiredWidth(125.0f)
+		.MaxDesiredWidth(125.0f)
+		[
+			SNew(SNumericEntryBox<float>)
+			.Value(this, &FCineSplineMetadataDetails::GetFocusDistance)
+			.AllowSpin(true)
+			.OnBeginSliderMovement(this, &FCineSplineMetadataDetails::OnBeginSliderMovement)
+			.OnEndSliderMovement(this, &FCineSplineMetadataDetails::OnEndSliderMovement)
+			.MinValue(TOptional<float>())
+			.MaxValue(TOptional<float>())
+			.MinSliderValue(TOptional<float>()) // No lower limit
+			.MaxSliderValue(TOptional<float>()) // No upper limit
+			.UndeterminedString(LOCTEXT("Multiple", "Multiple"))
+			.OnValueCommitted(this, &FCineSplineMetadataDetails::OnSetFocusDistance)
+			.OnValueChanged(this, &FCineSplineMetadataDetails::OnSetFocusDistance, ETextCommit::Default)
+			.Font(IDetailLayoutBuilder::GetDetailFont())
+		];
+
+		DetailGroup.AddWidgetRow()
+			.RowTag("CameraRotation")
+			.Visibility(TAttribute<EVisibility>(this, &FCineSplineMetadataDetails::IsEnabled))
+			.NameContent()
+			.HAlign(HAlign_Left)
+			.VAlign(VAlign_Center)
+			[
+				SNew(STextBlock)
+				.Text(LOCTEXT("CameraRotation", "Camera Rotation"))
+				.Font(IDetailLayoutBuilder::GetDetailFont())
+			]
+			.ValueContent()
+			.MinDesiredWidth(125.0f)
+			.MaxDesiredWidth(125.0f)
+			[
+				SNew(SRotatorInputBox)
+				.Roll(this, &FCineSplineMetadataDetails::GetRotation, EAxis::X)
+				.Pitch(this, &FCineSplineMetadataDetails::GetRotation, EAxis::Y)
+				.Yaw(this, &FCineSplineMetadataDetails::GetRotation, EAxis::Z)
+				.OnRollChanged(this, &FCineSplineMetadataDetails::OnSetRotation, ETextCommit::Default, EAxis::X)
+				.OnPitchChanged(this, &FCineSplineMetadataDetails::OnSetRotation, ETextCommit::Default, EAxis::Y)
+				.OnYawChanged(this, &FCineSplineMetadataDetails::OnSetRotation, ETextCommit::Default, EAxis::Z)
+				.OnRollCommitted(this, &FCineSplineMetadataDetails::OnSetRotation, EAxis::X)
+				.OnPitchCommitted(this, &FCineSplineMetadataDetails::OnSetRotation, EAxis::Y)
+				.OnYawCommitted(this, &FCineSplineMetadataDetails::OnSetRotation, EAxis::Z)
+				.AllowSpin(false)
+				.bColorAxisLabels(false)
+				.Font(IDetailLayoutBuilder::GetDetailFont())
+			];
+}
+
+#undef LOCTEXT_NAMESPACE
\ No newline at end of file
diff --git a/Engine/Plugins/Experimental/CineCameraRigs/Source/CineCameraRigsEditor/Private/CineSplineMetadataDetails.h b/Engine/Plugins/Experimental/CineCameraRigs/Source/CineCameraRigsEditor/Private/CineSplineMetadataDetails.h
new file mode 100644
--- /dev/null
+++ b/Engine/Plugins/Experimental/CineCameraRigs/Source/CineCameraRigsEditor/Private/CineSplineMetadataDetails.h
@@ -0,0 +1,71 @@
+// Copyright Epic Games, Inc. All Rights Reserved.
+
+#pragma once
+
+#include "CoreMinimal.h"
+#include "UObject/ObjectMacros.h"
+#include "Layout/Visibility.h"
+#include "Types/SlateEnums.h"
+#include "ScopedTransaction.h"
+
+#include "SplineMetadataDetailsFactory.h"
+
+#include "CineSplineMetadataDetails.generated.h"
+
+class UCineSplineMetadata;
+
+UCLASS()
+class UCineSplineMetadataDetailsFactory : public USplineMetadataDetailsFactoryBase
+{
+	GENERATED_BODY()
+
+public:
+	virtual TSharedPtr<ISplineMetadataDetails> Create() override;
+	virtual UClass* GetMetadataClass() const override;
+};
+
+class FCineSplineMetadataDetails : public ISplineMetadataDetails, public TSharedFromThis<FCineSplineMetadataDetails>
+{
+public:
+	virtual ~FCineSplineMetadataDetails() = default;
+	virtual FName GetName() const override { return FName("CineSplineMetadataDetails"); }
+	virtual FText GetDisplayName() const override;
+	virtual void Update(USplineComponent* InSplineComponent, const TSet<int32>& InSelectedKeys) override;
+	virtual void GenerateChildContent(IDetailGroup& InGroup) override;
+
+	TOptional<float> CustomPositionValue;
+	TOptional<float> FocalLengthValue;
+	TOptional<float> ApertureValue;
+	TOptional<float> FocusDistanceValue;
+	TOptional<FQuat> CameraRotationValue;
+
+	USplineComponent* SplineComp = nullptr;
+	TSet<int32> SelectedKeys;
+
+private:
+	void OnBeginSliderMovement();
+	void OnEndSliderMovement(float NewValue);
+
+	EVisibility IsEnabled() const { return (SelectedKeys.Num() > 0) ? EVisibility::Visible : EVisibility::Collapsed; }
+	bool IsOnePointSelected() const { return SelectedKeys.Num() == 1; }
+
+
+	TOptional<float> GetCustomPosition() const { return CustomPositionValue; }
+	void OnSetCustomPosition(float NewValue, ETextCommit::Type CommitInfo);
+
+	TOptional<float> GetFocalLength() const { return FocalLengthValue; }
+	void OnSetFocalLength(float NewValue, ETextCommit::Type CommitInfo);
+
+	TOptional<float> GetAperture() const { return ApertureValue; }
+	void OnSetAperture(float NewValue, ETextCommit::Type CommitInfo);
+
+	TOptional<float> GetFocusDistance() const { return FocusDistanceValue; }
+	void OnSetFocusDistance(float NewValue, ETextCommit::Type CommitInfo);
+
+	TOptional<float> GetRotation(EAxis::Type Axis) const;
+	void OnSetRotation(float NewValue, ETextCommit::Type CommitInfo, EAxis::Type Axis);
+
+	UCineSplineMetadata* GetMetadata() const;
+private:
+	TUniquePtr<FScopedTransaction> EditSliderValueTransaction;
+};
diff --git a/Engine/Plugins/Experimental/VirtualProductionUtilities/Config/BaseVirtualProductionUtilities.ini b/Engine/Plugins/Experimental/VirtualProductionUtilities/Config/BaseVirtualProductionUtilities.ini
--- a/Engine/Plugins/Experimental/VirtualProductionUtilities/Config/BaseVirtualProductionUtilities.ini
+++ b/Engine/Plugins/Experimental/VirtualProductionUtilities/Config/BaseVirtualProductionUtilities.ini
@@ -10,7 +10,6 @@ ScoutingSubsystemEditorUtilityClassPath=/VirtualProductionUtilities/BP_VPScoutin
 GestureManagerEditorUtilityClassPath=/VirtualProductionUtilities/VRTools/BP_GestureManager.BP_GestureManager_C
 +AdditionnalClassToLoad=/VirtualProductionUtilities/VR/VirtualScoutingInteractor.VirtualScoutingInteractor_C
 +AdditionnalClassToLoad=/VirtualProductionUtilities/VR/VirtualScoutingTeleporter.VirtualScoutingTeleporter_C
-VPSplinePreviewMeshPath=/VirtualProductionUtilities/Editor/Meshes/SM_VPSplineArraw.SM_VPSplineArraw
 
 
 [/Script/VPUtilities.VPBookmarkSettings]
diff --git a/Engine/Plugins/Experimental/VirtualProductionUtilities/Source/VPSpline/Private/VPSpline.cpp b/Engine/Plugins/Experimental/VirtualProductionUtilities/Source/VPSpline/Private/VPSpline.cpp
deleted file mode 100644
diff --git a/Engine/Plugins/Experimental/VirtualProductionUtilities/Source/VPSpline/Private/VPSplineActor.cpp b/Engine/Plugins/Experimental/VirtualProductionUtilities/Source/VPSpline/Private/VPSplineActor.cpp
deleted file mode 100644
diff --git a/Engine/Plugins/Experimental/VirtualProductionUtilities/Source/VPSpline/Private/VPSplineComponent.cpp b/Engine/Plugins/Experimental/VirtualProductionUtilities/Source/VPSpline/Private/VPSplineComponent.cpp
deleted file mode 100644
diff --git a/Engine/Plugins/Experimental/VirtualProductionUtilities/Source/VPSpline/Private/VPSplineLog.h b/Engine/Plugins/Experimental/VirtualProductionUtilities/Source/VPSpline/Private/VPSplineLog.h
deleted file mode 100644
diff --git a/Engine/Plugins/Experimental/VirtualProductionUtilities/Source/VPSpline/Private/VPSplineMetadata.cpp b/Engine/Plugins/Experimental/VirtualProductionUtilities/Source/VPSpline/Private/VPSplineMetadata.cpp
deleted file mode 100644
diff --git a/Engine/Plugins/Experimental/VirtualProductionUtilities/Source/VPSpline/Public/VPSpline.h b/Engine/Plugins/Experimental/VirtualProductionUtilities/Source/VPSpline/Public/VPSpline.h
deleted file mode 100644
diff --git a/Engine/Plugins/Experimental/VirtualProductionUtilities/Source/VPSpline/Public/VPSplineActor.h b/Engine/Plugins/Experimental/VirtualProductionUtilities/Source/VPSpline/Public/VPSplineActor.h
deleted file mode 100644
diff --git a/Engine/Plugins/Experimental/VirtualProductionUtilities/Source/VPSpline/Public/VPSplineComponent.h b/Engine/Plugins/Experimental/VirtualProductionUtilities/Source/VPSpline/Public/VPSplineComponent.h
deleted file mode 100644
diff --git a/Engine/Plugins/Experimental/VirtualProductionUtilities/Source/VPSpline/Public/VPSplineMetadata.h b/Engine/Plugins/Experimental/VirtualProductionUtilities/Source/VPSpline/Public/VPSplineMetadata.h
deleted file mode 100644
diff --git a/Engine/Plugins/Experimental/VirtualProductionUtilities/Source/VPSpline/Public/VPSplinePointData.h b/Engine/Plugins/Experimental/VirtualProductionUtilities/Source/VPSpline/Public/VPSplinePointData.h
deleted file mode 100644
diff --git a/Engine/Plugins/Experimental/VirtualProductionUtilities/Source/VPSpline/VPSpline.Build.cs b/Engine/Plugins/Experimental/VirtualProductionUtilities/Source/VPSpline/VPSpline.Build.cs
deleted file mode 100644
diff --git a/Engine/Plugins/Experimental/VirtualProductionUtilities/Source/VPSplineEditor/Private/VPSplineComponentVisualizer.cpp b/Engine/Plugins/Experimental/VirtualProductionUtilities/Source/VPSplineEditor/Private/VPSplineComponentVisualizer.cpp
deleted file mode 100644
diff --git a/Engine/Plugins/Experimental/VirtualProductionUtilities/Source/VPSplineEditor/Private/VPSplineComponentVisualizer.h b/Engine/Plugins/Experimental/VirtualProductionUtilities/Source/VPSplineEditor/Private/VPSplineComponentVisualizer.h
deleted file mode 100644
diff --git a/Engine/Plugins/Experimental/VirtualProductionUtilities/Source/VPSplineEditor/Private/VPSplineEditor.cpp b/Engine/Plugins/Experimental/VirtualProductionUtilities/Source/VPSplineEditor/Private/VPSplineEditor.cpp
deleted file mode 100644
diff --git a/Engine/Plugins/Experimental/VirtualProductionUtilities/Source/VPSplineEditor/Private/VPSplineMetadataDetails.cpp b/Engine/Plugins/Experimental/VirtualProductionUtilities/Source/VPSplineEditor/Private/VPSplineMetadataDetails.cpp
deleted file mode 100644
diff --git a/Engine/Plugins/Experimental/VirtualProductionUtilities/Source/VPSplineEditor/Private/VPSplineMetadataDetails.h b/Engine/Plugins/Experimental/VirtualProductionUtilities/Source/VPSplineEditor/Private/VPSplineMetadataDetails.h
deleted file mode 100644
diff --git a/Engine/Plugins/Experimental/VirtualProductionUtilities/Source/VPSplineEditor/VPSplineEditor.Build.cs b/Engine/Plugins/Experimental/VirtualProductionUtilities/Source/VPSplineEditor/VPSplineEditor.Build.cs
deleted file mode 100644
diff --git a/Engine/Plugins/Experimental/VirtualProductionUtilities/VirtualProductionUtilities.uplugin b/Engine/Plugins/Experimental/VirtualProductionUtilities/VirtualProductionUtilities.uplugin
--- a/Engine/Plugins/Experimental/VirtualProductionUtilities/VirtualProductionUtilities.uplugin
+++ b/Engine/Plugins/Experimental/VirtualProductionUtilities/VirtualProductionUtilities.uplugin
@@ -38,16 +38,6 @@
 			"Name": "VPMaterialsEditor",
 			"Type": "Editor",
 			"LoadingPhase": "PostDefault"
-		},
-		{
-			"Name": "VPSpline",
-			"Type": "Runtime",
-			"LoadingPhase": "Default"
-		},
-		{
-			"Name": "VPSplineEditor",
-			"Type": "Editor",
-			"LoadingPhase": "PostEngineInit"
 		}
 	],
 	"Plugins": [
