Check CineSplineComponent nullptr

diff --git a/Engine/Plugins/Experimental/CineCameraRigs/Source/CineCameraRigs/Private/CineCameraRigRail.cpp b/Engine/Plugins/Experimental/CineCameraRigs/Source/CineCameraRigs/Private/CineCameraRigRail.cpp
--- a/Engine/Plugins/Experimental/CineCameraRigs/Source/CineCameraRigs/Private/CineCameraRigRail.cpp
+++ b/Engine/Plugins/Experimental/CineCameraRigs/Source/CineCameraRigs/Private/CineCameraRigRail.cpp
@@ -48,6 +48,11 @@ ACineCameraRigRail::ACineCameraRigRail(const FObjectInitializer& ObjectInitializ
 
 void ACineCameraRigRail::UpdateRailComponents()
 {
+	if (CineSplineComponent == nullptr)
+	{
+		return;
+	}
+
 	if (!GetWorld())
 	{
 		return;
