Attempt to fix nonunity compile errors for some platforms

diff --git a/Engine/Plugins/Experimental/CineCameraRigs/Source/CineCameraRigs/Public/CineCameraRigsSettings.h b/Engine/Plugins/Experimental/CineCameraRigs/Source/CineCameraRigs/Public/CineCameraRigsSettings.h
--- a/Engine/Plugins/Experimental/CineCameraRigs/Source/CineCameraRigs/Public/CineCameraRigsSettings.h
+++ b/Engine/Plugins/Experimental/CineCameraRigs/Source/CineCameraRigs/Public/CineCameraRigsSettings.h
@@ -7,6 +7,9 @@
 
 #include "CineCameraRigsSettings.generated.h"
 
+class UMaterialInterface;
+class UTexture2D;
+
 UCLASS(BlueprintType, Config = CineCameraRigs)
 class CINECAMERARIGS_API UCineCameraRigRailSettings : public UObject
 {
