CineCameraRigRail details customization

diff --git a/Engine/Plugins/Experimental/CineCameraRigs/Source/CineCameraRigs/Private/CineCameraRigRail.cpp b/Engine/Plugins/Experimental/CineCameraRigs/Source/CineCameraRigs/Private/CineCameraRigRail.cpp
--- a/Engine/Plugins/Experimental/CineCameraRigs/Source/CineCameraRigs/Private/CineCameraRigRail.cpp
+++ b/Engine/Plugins/Experimental/CineCameraRigs/Source/CineCameraRigs/Private/CineCameraRigRail.cpp
@@ -58,7 +58,7 @@ void ACineCameraRigRail::UpdateRailComponents()
 		return;
 	}
 
-	if (!bUseCustomPosition)
+	if (!bUseAbsolutePosition)
 	{
 		Super::UpdateRailComponents();
 	}
@@ -67,13 +67,13 @@ void ACineCameraRigRail::UpdateRailComponents()
 		USceneComponent* AttachComponent = GetDefaultAttachComponent();
 		if (CineSplineComponent && AttachComponent)
 		{
-			float InputKey = CineSplineComponent->GetInputKeyAtPosition(CustomPosition);
+			float InputKey = CineSplineComponent->GetInputKeyAtPosition(AbsolutePositionOnRail);
 			FVector const SplinePosition = CineSplineComponent->GetLocationAtSplineInputKey(InputKey, ESplineCoordinateSpace::World);
 			FQuat SplineQuat = CineSplineComponent->GetQuaternionAtSplineInputKey(InputKey, ESplineCoordinateSpace::World);
 
-			if (bUseCameraRotationMetadata)
+			if (bUsePointRotation)
 			{
-				SplineQuat = CineSplineComponent->GetComponentTransform().GetRotation() * CineSplineComponent->GetCameraRotationAtSplineInputKey(InputKey);
+				SplineQuat = CineSplineComponent->GetComponentTransform().GetRotation() * CineSplineComponent->GetPointRotationAtSplineInputKey(InputKey);
 			}
 			FVector Position = GetActorLocation();
 			FRotator Rotation = GetActorRotation();
@@ -114,9 +114,9 @@ void ACineCameraRigRail::UpdateRailComponents()
 	TArray< AActor* > AttachedActors;
 	GetAttachedActors(AttachedActors);
 	float InputKey = 0.0f;
-	if (bUseCustomPosition)
+	if (bUseAbsolutePosition)
 	{
-		InputKey = CineSplineComponent->GetInputKeyAtPosition(CustomPosition);
+		InputKey = CineSplineComponent->GetInputKeyAtPosition(AbsolutePositionOnRail);
 	}
 	else
 	{
@@ -130,15 +130,15 @@ void ACineCameraRigRail::UpdateRailComponents()
 	{
 		if (UCineCameraComponent* CameraComponent = Cast<UCineCameraComponent>(AttachedActor->GetComponentByClass(UCineCameraComponent::StaticClass())))
 		{
-			if (bDriveFocusDistance)
+			if (bInheritFocusDistance)
 			{
 				CameraComponent->FocusSettings.ManualFocusDistance = FocusDistanceValue;
 			}
-			if (bDriveFocalLength)
+			if (bInheritFocalLength)
 			{
 				CameraComponent->SetCurrentFocalLength(FocalLengthValue);
 			}
-			if (bDriveAperture)
+			if (bInheritAperture)
 			{
 				CameraComponent->SetCurrentAperture(ApertureValue);
 			}
@@ -150,8 +150,8 @@ void ACineCameraRigRail::UpdateRailComponents()
 		const UWorld* const MyWorld = GetWorld();
 		if (MyWorld && !MyWorld->IsGameWorld())
 		{
-			// If bUseCustomPosition is false, UpdatePreviewMeshes() is already called inside Super::UpdateRailComponents()
-			if (bUseCustomPosition)
+			// If bUseAbsolutePosition is false, UpdatePreviewMeshes() is already called inside Super::UpdateRailComponents()
+			if (bUseAbsolutePosition)
 			{
 				UpdatePreviewMeshes();
 			}
@@ -258,11 +258,11 @@ FVector ACineCameraRigRail::GetVelocityAtPosition(const float InPosition, const
 
 	float TimeMultiplier = 1.0f / CineSplineComponent->Duration;
 
-	if (bUseCustomPosition)
+	if (bUseAbsolutePosition)
 	{
 		int32 const NumPoints = CineSplineComponent->GetNumberOfSplinePoints();
-		float const MinPosition = CineSplineComponent->GetFloatPropertyAtSplinePoint(0, FName(TEXT("CustomPosition")));
-		float const MaxPosition = CineSplineComponent->GetFloatPropertyAtSplinePoint(NumPoints - 1, FName(TEXT("CustomPosition")));
+		float const MinPosition = CineSplineComponent->GetFloatPropertyAtSplinePoint(0, FName(TEXT("AbsolutePosition")));
+		float const MaxPosition = CineSplineComponent->GetFloatPropertyAtSplinePoint(NumPoints - 1, FName(TEXT("AbsolutePosition")));
 		float const t0 = InPosition + delta >= MaxPosition ? InPosition - delta : InPosition;
 		float const t1 = InPosition + delta >= MaxPosition ? InPosition : InPosition + delta;
 		float const InKey0 = CineSplineComponent->GetInputKeyAtPosition(t0);
diff --git a/Engine/Plugins/Experimental/CineCameraRigs/Source/CineCameraRigs/Private/CineSplineComponent.cpp b/Engine/Plugins/Experimental/CineCameraRigs/Source/CineCameraRigs/Private/CineSplineComponent.cpp
--- a/Engine/Plugins/Experimental/CineCameraRigs/Source/CineCameraRigs/Private/CineSplineComponent.cpp
+++ b/Engine/Plugins/Experimental/CineCameraRigs/Source/CineCameraRigs/Private/CineSplineComponent.cpp
@@ -111,16 +111,16 @@ void UCineSplineComponent::SynchronizeProperties()
 	{
 		CineSplineMetadata->Fixup(NumOfPoints, this);
 
-		// Fixing invalid CustomPosition
+		// Fixing invalid AbsolutePosition
 		// This is to make sure CustomPostion value is incrementing from previous point
-		float PreviousValue = CineSplineMetadata->CustomPosition.Points[0].OutVal;
+		float PreviousValue = CineSplineMetadata->AbsolutePosition.Points[0].OutVal;
 		for (int32 Index = 1; Index < NumOfPoints; ++Index)
 		{
-			float CurrentValue = CineSplineMetadata->CustomPosition.Points[Index].OutVal;
+			float CurrentValue = CineSplineMetadata->AbsolutePosition.Points[Index].OutVal;
 			if (CurrentValue <= PreviousValue)
 			{
 				CurrentValue = PreviousValue + 1.0f;
-				CineSplineMetadata->CustomPosition.Points[Index].OutVal = CurrentValue;
+				CineSplineMetadata->AbsolutePosition.Points[Index].OutVal = CurrentValue;
 			}
 			PreviousValue = CurrentValue;
 		}
@@ -152,28 +152,28 @@ void UCineSplineComponent::SetFocusDistanceAtSplinePoint(const int32 PointIndex,
 }
 
 
-void UCineSplineComponent::SetCustomPositionAtSplinePoint(const int32 PointIndex, const float Value)
+void UCineSplineComponent::SetAbsolutePositionAtSplinePoint(const int32 PointIndex, const float Value)
 {
-	int32 NumPoints = CineSplineMetadata->CustomPosition.Points.Num();
+	int32 NumPoints = CineSplineMetadata->AbsolutePosition.Points.Num();
 	check(PointIndex >= 0 && PointIndex < NumPoints);
 	CineSplineMetadata->Modify();
-	CineSplineMetadata->CustomPosition.Points[PointIndex].OutVal = Value;
+	CineSplineMetadata->AbsolutePosition.Points[PointIndex].OutVal = Value;
 }
 
-void UCineSplineComponent::SetCameraRotationAtSplinePoint(const int32 PointIndex, const FQuat Value)
+void UCineSplineComponent::SetPointRotationAtSplinePoint(const int32 PointIndex, const FQuat Value)
 {
-	int32 NumPoints = CineSplineMetadata->CameraRotation.Points.Num();
+	int32 NumPoints = CineSplineMetadata->PointRotation.Points.Num();
 	check(PointIndex >= 0 && PointIndex < NumPoints);
 	CineSplineMetadata->Modify();
-	CineSplineMetadata->CameraRotation.Points[PointIndex].OutVal = Value;
+	CineSplineMetadata->PointRotation.Points[PointIndex].OutVal = Value;
 }
 
 bool UCineSplineComponent::FindSplineDataAtPosition(const float InPosition, int32& OutIndex, const float Tolerance) const
 {
-	int32 NumPoints = CineSplineMetadata->CustomPosition.Points.Num();
+	int32 NumPoints = CineSplineMetadata->AbsolutePosition.Points.Num();
 	for (int32 i = 0; i < NumPoints; ++i)
 	{
-		if (FMath::IsNearlyEqual(InPosition, CineSplineMetadata->CustomPosition.Points[i].OutVal, Tolerance))
+		if (FMath::IsNearlyEqual(InPosition, CineSplineMetadata->AbsolutePosition.Points[i].OutVal, Tolerance))
 		{
 			OutIndex = i;
 			return true;
@@ -189,15 +189,15 @@ float UCineSplineComponent::GetInputKeyAtPosition(const float InPosition) const
 
 	float OutValue = 0.0f;
 	
-	int32 NumPoints = CineSplineMetadata->CustomPosition.Points.Num();
+	int32 NumPoints = CineSplineMetadata->AbsolutePosition.Points.Num();
 	for (int32 i = 0; i < NumPoints; ++i)
 	{
-		if (InPosition < CineSplineMetadata->CustomPosition.Points[i].OutVal)
+		if (InPosition < CineSplineMetadata->AbsolutePosition.Points[i].OutVal)
 		{
 			if (i > 0)
 			{
-				float Value0 = CineSplineMetadata->CustomPosition.Points[i - 1].OutVal;
-				float Value1 = CineSplineMetadata->CustomPosition.Points[i].OutVal;
+				float Value0 = CineSplineMetadata->AbsolutePosition.Points[i - 1].OutVal;
+				float Value1 = CineSplineMetadata->AbsolutePosition.Points[i].OutVal;
 				OutValue = (InPosition - Value0) / (Value1 - Value0) + (float)(i-1);
 			}
 			break;
@@ -210,24 +210,24 @@ float UCineSplineComponent::GetInputKeyAtPosition(const float InPosition) const
 
 float UCineSplineComponent::GetPositionAtInputKey(const float InKey) const
 {
-	return GetFloatPropertyAtSplineInputKey(InKey, FName(TEXT("CustomPosition")));
+	return GetFloatPropertyAtSplineInputKey(InKey, FName(TEXT("AbsolutePosition")));
 }
 
-FQuat UCineSplineComponent::GetCameraRotationAtSplinePoint(int32 Index) const
+FQuat UCineSplineComponent::GetPointRotationAtSplinePoint(int32 Index) const
 {
-	const int32 NumPoints = CineSplineMetadata->CameraRotation.Points.Num();
+	const int32 NumPoints = CineSplineMetadata->PointRotation.Points.Num();
 	if (NumPoints > 0)
 	{
 		const int32 ClampedIndex = FMath::Clamp(Index, 0, NumPoints - 1);
-		return CineSplineMetadata->CameraRotation.Points[ClampedIndex].OutVal;
+		return CineSplineMetadata->PointRotation.Points[ClampedIndex].OutVal;
 	}
 	
 	return FQuat::Identity;
 }
 
-FQuat UCineSplineComponent::GetCameraRotationAtSplineInputKey(float InKey) const
+FQuat UCineSplineComponent::GetPointRotationAtSplineInputKey(float InKey) const
 {
-	return CineSplineMetadata->CameraRotation.Eval(InKey, FQuat::Identity);
+	return CineSplineMetadata->PointRotation.Eval(InKey, FQuat::Identity);
 }
 
 void UCineSplineComponent::UpdateSplineDataAtIndex(const int InIndex, const FCineSplinePointData& InPointData)
@@ -240,17 +240,17 @@ void UCineSplineComponent::UpdateSplineDataAtIndex(const int InIndex, const FCin
 	SetApertureAtSplinePoint(InIndex, InPointData.Aperture);
 	SetFocusDistanceAtSplinePoint(InIndex, InPointData.FocusDistance);
 	const FQuat Quat = GetComponentTransform().InverseTransformRotation(InPointData.Rotation.Quaternion());
-	SetCameraRotationAtSplinePoint(InIndex, Quat);
+	SetPointRotationAtSplinePoint(InIndex, Quat);
 	UpdateSpline();
 }
 
 void UCineSplineComponent::AddSplineDataAtPosition(const float InPosition, const FCineSplinePointData& InPointData)
 {
 	int32 NewIndex = 0;
-	int32 NumPoints = CineSplineMetadata->CustomPosition.Points.Num();
+	int32 NumPoints = CineSplineMetadata->AbsolutePosition.Points.Num();
 	for (int32 i = 0; i < NumPoints; ++i)
 	{
-		if (InPosition <= CineSplineMetadata->CustomPosition.Points[i].OutVal)
+		if (InPosition <= CineSplineMetadata->AbsolutePosition.Points[i].OutVal)
 		{
 			break;
 		}
@@ -261,10 +261,10 @@ void UCineSplineComponent::AddSplineDataAtPosition(const float InPosition, const
 	SetFocalLengthAtSplinePoint(NewIndex, InPointData.FocalLength);
 	SetApertureAtSplinePoint(NewIndex, InPointData.Aperture);
 	SetFocusDistanceAtSplinePoint(NewIndex, InPointData.FocusDistance);
-	SetCustomPositionAtSplinePoint(NewIndex, InPosition);
+	SetAbsolutePositionAtSplinePoint(NewIndex, InPosition);
 	SetSplinePointType(NewIndex, ESplinePointType::Curve, false);
 	const FQuat Quat = GetComponentTransform().InverseTransformRotation(InPointData.Rotation.Quaternion());
-	SetCameraRotationAtSplinePoint(NewIndex, Quat);
+	SetPointRotationAtSplinePoint(NewIndex, Quat);
 	UpdateSpline();
 }
 
@@ -273,7 +273,7 @@ FCineSplinePointData UCineSplineComponent::GetSplineDataAtPosition(const float I
 	FCineSplinePointData PointData;
 	const float InputKey = GetInputKeyAtPosition(InPosition);
 	PointData.Location = GetLocationAtSplineInputKey(InputKey, ESplineCoordinateSpace::World);
-	PointData.Rotation = GetCameraRotationAtSplineInputKey(InputKey).Rotator();
+	PointData.Rotation = GetPointRotationAtSplineInputKey(InputKey).Rotator();
 	PointData.FocalLength = GetFloatPropertyAtSplineInputKey(InputKey, FName(TEXT("FocalLength")));
 	PointData.Aperture = GetFloatPropertyAtSplineInputKey(InputKey, FName(TEXT("Aperture")));
 	PointData.FocusDistance = GetFloatPropertyAtSplineInputKey(InputKey, FName(TEXT("FocusDistance")));
diff --git a/Engine/Plugins/Experimental/CineCameraRigs/Source/CineCameraRigs/Private/CineSplineMetadata.cpp b/Engine/Plugins/Experimental/CineCameraRigs/Source/CineCameraRigs/Private/CineSplineMetadata.cpp
--- a/Engine/Plugins/Experimental/CineCameraRigs/Source/CineCameraRigs/Private/CineSplineMetadata.cpp
+++ b/Engine/Plugins/Experimental/CineCameraRigs/Source/CineCameraRigs/Private/CineSplineMetadata.cpp
@@ -29,23 +29,23 @@ void UCineSplineMetadata::InsertPoint(int32 Index, float Alpha, bool bClosedLoop
 		float NewFocalLengthVal = FocalLength.Points[Index].OutVal;
 		float NewApertureVal = Aperture.Points[Index].OutVal;
 		float NewFocusDistanceVal = FocusDistance.Points[Index].OutVal;
-		float NewCustomPositionVal = CustomPosition.Points[Index].OutVal;
-		FQuat NewRotationVal = CameraRotation.Points[Index].OutVal;
+		float NewAbsolutePositionVal = AbsolutePosition.Points[Index].OutVal;
+		FQuat NewRotationVal = PointRotation.Points[Index].OutVal;
 		if (bHasPrevIndex)
 		{
 			float PrevFocalLengthVal = FocalLength.Points[PrevIndex].OutVal;
 			float PrevApertureVal = Aperture.Points[PrevIndex].OutVal;
 			float PrevFocusDistanceVal = FocusDistance.Points[PrevIndex].OutVal;
-			float PrevCustomPositionVal = CustomPosition.Points[PrevIndex].OutVal;
-			FQuat PrevRotationVal = CameraRotation.Points[PrevIndex].OutVal;
+			float PrevAbsolutePositionVal = AbsolutePosition.Points[PrevIndex].OutVal;
+			FQuat PrevRotationVal = PointRotation.Points[PrevIndex].OutVal;
 			NewFocalLengthVal = FMath::LerpStable(PrevFocalLengthVal, NewFocalLengthVal, Alpha);
 			NewApertureVal = FMath::LerpStable(PrevApertureVal, NewApertureVal, Alpha);
 			NewFocusDistanceVal = FMath::LerpStable(PrevFocusDistanceVal, NewFocusDistanceVal, Alpha);
-			NewCustomPositionVal = FMath::LerpStable(PrevCustomPositionVal, NewCustomPositionVal, Alpha);
+			NewAbsolutePositionVal = FMath::LerpStable(PrevAbsolutePositionVal, NewAbsolutePositionVal, Alpha);
 			NewRotationVal = FMath::LerpStable(PrevRotationVal, NewRotationVal, Alpha);
 		}
-		FInterpCurvePoint<float> NewCustomPosition(InputKey, NewCustomPositionVal);
-		CustomPosition.Points.Insert(NewCustomPosition, Index);
+		FInterpCurvePoint<float> NewAbsolutePosition(InputKey, NewAbsolutePositionVal);
+		AbsolutePosition.Points.Insert(NewAbsolutePosition, Index);
 
 		FInterpCurvePoint<float> NewFocalLength(InputKey, NewFocalLengthVal);
 		FocalLength.Points.Insert(NewFocalLength, Index);
@@ -57,7 +57,7 @@ void UCineSplineMetadata::InsertPoint(int32 Index, float Alpha, bool bClosedLoop
 		FocusDistance.Points.Insert(NewFocusDistance, Index);
 
 		FInterpCurvePoint<FQuat> NewRotation(InputKey, NewRotationVal);
-		CameraRotation.Points.Insert(NewRotation, Index);
+		PointRotation.Points.Insert(NewRotation, Index);
 
 
 		for (int32 i = Index + 1; i < FocalLength.Points.Num(); ++i)
@@ -65,8 +65,8 @@ void UCineSplineMetadata::InsertPoint(int32 Index, float Alpha, bool bClosedLoop
 			FocalLength.Points[i].InVal += 1.0f;
 			Aperture.Points[i].InVal += 1.0f;
 			FocusDistance.Points[i].InVal += 1.0f;
-			CustomPosition.Points[i].InVal += 1.0f;
-			CameraRotation.Points[i].InVal += 1.0f;
+			AbsolutePosition.Points[i].InVal += 1.0f;
+			PointRotation.Points[i].InVal += 1.0f;
 		}
 	}
 }
@@ -91,22 +91,22 @@ void UCineSplineMetadata::UpdatePoint(int32 Index, float Alpha, bool bClosedLoop
 
 	if (bHasPrevIndex && bHasNextIndex)
 	{
-		float PrevCustomPositionVal = CustomPosition.Points[PrevIndex].OutVal;
+		float PrevAbsolutePositionVal = AbsolutePosition.Points[PrevIndex].OutVal;
 		float PrevFocalLengthVal = FocalLength.Points[PrevIndex].OutVal;
 		float PrevApertureVal = Aperture.Points[PrevIndex].OutVal;
 		float PrevFocusDistanceVal = FocusDistance.Points[PrevIndex].OutVal;
-		FQuat PrevRotationVal = CameraRotation.Points[PrevIndex].OutVal;
-		float NextCustomPositionVal = CustomPosition.Points[NextIndex].OutVal;
+		FQuat PrevRotationVal = PointRotation.Points[PrevIndex].OutVal;
+		float NextAbsolutePositionVal = AbsolutePosition.Points[NextIndex].OutVal;
 		float NextFocalLengthVal = FocalLength.Points[NextIndex].OutVal;
 		float NextApertureVal = Aperture.Points[NextIndex].OutVal;
 		float NextFocusDistanceVal = FocusDistance.Points[NextIndex].OutVal;
-		FQuat NextRotationVal = CameraRotation.Points[NextIndex].OutVal;
+		FQuat NextRotationVal = PointRotation.Points[NextIndex].OutVal;
 
-		CustomPosition.Points[Index].OutVal = FMath::LerpStable(PrevCustomPositionVal, NextCustomPositionVal, Alpha);
+		AbsolutePosition.Points[Index].OutVal = FMath::LerpStable(PrevAbsolutePositionVal, NextAbsolutePositionVal, Alpha);
 		FocalLength.Points[Index].OutVal = FMath::LerpStable(PrevFocalLengthVal, NextFocalLengthVal, Alpha);
 		Aperture.Points[Index].OutVal = FMath::LerpStable(PrevApertureVal, NextApertureVal, Alpha);
 		FocusDistance.Points[Index].OutVal = FMath::LerpStable(PrevFocusDistanceVal, NextFocusDistanceVal, Alpha);
-		CameraRotation.Points[Index].OutVal = FMath::LerpStable(PrevRotationVal, NextRotationVal, Alpha);
+		PointRotation.Points[Index].OutVal = FMath::LerpStable(PrevRotationVal, NextRotationVal, Alpha);
 	}
 }
 
@@ -114,7 +114,7 @@ void UCineSplineMetadata::AddPoint(float InputKey)
 {
 	Modify();
 
-	float NewCustomPositionVal = -1.0f;
+	float NewAbsolutePositionVal = -1.0f;
 	float NewFocalLengthVal = 35.0f;
 	float NewApertureVal = 2.8f;
 	float NewFocusDistanceVal = 100000.f;
@@ -126,16 +126,16 @@ void UCineSplineMetadata::AddPoint(float InputKey)
 		NewFocalLengthVal = FocalLength.Points[Index].OutVal;
 		NewApertureVal = Aperture.Points[Index].OutVal;
 		NewFocusDistanceVal = FocusDistance.Points[Index].OutVal;
-		NewCustomPositionVal = CustomPosition.Points[Index].OutVal + 1.0f;
-		NewRotationVal = CameraRotation.Points[Index].OutVal;
+		NewAbsolutePositionVal = AbsolutePosition.Points[Index].OutVal + 1.0f;
+		NewRotationVal = PointRotation.Points[Index].OutVal;
 	}
 
 	float NewInputKey = static_cast<float>(++Index);
 	FocalLength.Points.Emplace(NewInputKey, NewFocalLengthVal);
 	Aperture.Points.Emplace(NewInputKey, NewApertureVal);
 	FocusDistance.Points.Emplace(NewInputKey, NewFocusDistanceVal);
-	CustomPosition.Points.Emplace(NewInputKey, NewCustomPositionVal);
-	CameraRotation.Points.Emplace(NewInputKey, NewRotationVal);
+	AbsolutePosition.Points.Emplace(NewInputKey, NewAbsolutePositionVal);
+	PointRotation.Points.Emplace(NewInputKey, NewRotationVal);
 }
 
 void UCineSplineMetadata::RemovePoint(int32 Index)
@@ -143,19 +143,19 @@ void UCineSplineMetadata::RemovePoint(int32 Index)
 	check(Index < FocalLength.Points.Num());
 
 	Modify();
-	CustomPosition.Points.RemoveAt(Index);
+	AbsolutePosition.Points.RemoveAt(Index);
 	FocalLength.Points.RemoveAt(Index);
 	Aperture.Points.RemoveAt(Index);
 	FocusDistance.Points.RemoveAt(Index);
-	CameraRotation.Points.RemoveAt(Index);
+	PointRotation.Points.RemoveAt(Index);
 
 	for (int32 i = Index; i < FocalLength.Points.Num(); ++i)
 	{
-		CustomPosition.Points[i].InVal -= 1.0f;
+		AbsolutePosition.Points[i].InVal -= 1.0f;
 		FocalLength.Points[i].InVal -= 1.0f;
 		Aperture.Points[i].InVal -= 1.0f;
 		FocusDistance.Points[i].InVal -= 1.0f;
-		CameraRotation.Points[i].InVal -= 1.0f;
+		PointRotation.Points[i].InVal -= 1.0f;
 	}
 }
 
@@ -164,28 +164,28 @@ void UCineSplineMetadata::DuplicatePoint(int32 Index)
 	check(Index < FocalLength.Points.Num());
 
 	int32 NumPoints = FocalLength.Points.Num();
-	float NewValue = CustomPosition.Points[Index].OutVal + 1.0f;
+	float NewValue = AbsolutePosition.Points[Index].OutVal + 1.0f;
 	if (Index < NumPoints - 1)
 	{
-		NewValue = (CustomPosition.Points[Index].OutVal + CustomPosition.Points[Index + 1].OutVal) * 0.5;
+		NewValue = (AbsolutePosition.Points[Index].OutVal + AbsolutePosition.Points[Index + 1].OutVal) * 0.5;
 	}
 
 	Modify();
-	CustomPosition.Points.Insert(FInterpCurvePoint<float>(CustomPosition.Points[Index]), Index);
+	AbsolutePosition.Points.Insert(FInterpCurvePoint<float>(AbsolutePosition.Points[Index]), Index);
 	FocalLength.Points.Insert(FInterpCurvePoint<float>(FocalLength.Points[Index]), Index);
 	Aperture.Points.Insert(FInterpCurvePoint<float>(Aperture.Points[Index]), Index);
 	FocusDistance.Points.Insert(FInterpCurvePoint<float>(FocusDistance.Points[Index]), Index);
-	CameraRotation.Points.Insert(FInterpCurvePoint<FQuat>(CameraRotation.Points[Index]), Index);
-	CustomPosition.Points[Index+1].OutVal = NewValue;
+	PointRotation.Points.Insert(FInterpCurvePoint<FQuat>(PointRotation.Points[Index]), Index);
+	AbsolutePosition.Points[Index+1].OutVal = NewValue;
 
 
 	for (int32 i = Index + 1; i < FocalLength.Points.Num(); ++i)
 	{
-		CustomPosition.Points[i].InVal += 1.0f;
+		AbsolutePosition.Points[i].InVal += 1.0f;
 		FocalLength.Points[i].InVal += 1.0f;
 		Aperture.Points[i].InVal += 1.0f;
 		FocusDistance.Points[i].InVal += 1.0f;
-		CameraRotation.Points[i].InVal += 1.0f;
+		PointRotation.Points[i].InVal += 1.0f;
 	}
 
 }
@@ -203,8 +203,8 @@ void UCineSplineMetadata::CopyPoint(const USplineMetadata* FromSplineMetadata, i
 		FocalLength.Points[ToIndex].OutVal = FromMetadata->FocalLength.Points[FromIndex].OutVal;
 		Aperture.Points[ToIndex].OutVal = FromMetadata->Aperture.Points[FromIndex].OutVal;
 		FocusDistance.Points[ToIndex].OutVal = FromMetadata->FocusDistance.Points[FromIndex].OutVal;
-		CustomPosition.Points[ToIndex].OutVal = FromMetadata->CustomPosition.Points[FromIndex].OutVal;
-		CameraRotation.Points[ToIndex].OutVal = FromMetadata->CameraRotation.Points[FromIndex].OutVal;
+		AbsolutePosition.Points[ToIndex].OutVal = FromMetadata->AbsolutePosition.Points[FromIndex].OutVal;
+		PointRotation.Points[ToIndex].OutVal = FromMetadata->PointRotation.Points[FromIndex].OutVal;
 	}
 }
 
@@ -214,8 +214,8 @@ void UCineSplineMetadata::Reset(int32 NumPoints)
 	FocalLength.Points.Reset(NumPoints);
 	Aperture.Points.Reset(NumPoints);
 	FocusDistance.Points.Reset(NumPoints);
-	CustomPosition.Points.Reset(NumPoints);
-	CameraRotation.Points.Reset(NumPoints);
+	AbsolutePosition.Points.Reset(NumPoints);
+	PointRotation.Points.Reset(NumPoints);
 }
 
 #if WITH_EDITORONLY_DATA
@@ -250,7 +250,7 @@ void UCineSplineMetadata::Fixup(int32 NumPoints, USplineComponent* SplineComp)
 	FixupCurve(FocalLength, Defaults.DefaultFocalLength, NumPoints);
 	FixupCurve(Aperture, Defaults.DefaultAperture, NumPoints);
 	FixupCurve(FocusDistance, Defaults.DefaultFocusDistance, NumPoints);
-	FixupCurve(CustomPosition, Defaults.DefaultCustomPosition, NumPoints);
-	FixupCurve(CameraRotation, Defaults.DefaultCameraRotation, NumPoints);
+	FixupCurve(AbsolutePosition, Defaults.DefaultAbsolutePosition, NumPoints);
+	FixupCurve(PointRotation, Defaults.DefaultPointRotation, NumPoints);
 #endif
 }
\ No newline at end of file
diff --git a/Engine/Plugins/Experimental/CineCameraRigs/Source/CineCameraRigs/Public/CineCameraRigRail.h b/Engine/Plugins/Experimental/CineCameraRigs/Source/CineCameraRigs/Public/CineCameraRigRail.h
--- a/Engine/Plugins/Experimental/CineCameraRigs/Source/CineCameraRigs/Public/CineCameraRigRail.h
+++ b/Engine/Plugins/Experimental/CineCameraRigs/Source/CineCameraRigs/Public/CineCameraRigRail.h
@@ -21,24 +21,24 @@ public:
 	UFUNCTION(BlueprintPure, Category = "Rail Components")
 	UCineSplineComponent* GetCineSplineComponent() const { return CineSplineComponent; }
 
-	/* Use CustomPosition metadata to parameterize the spline*/
-	UPROPERTY(EditAnywhere, BlueprintReadWrite, Category = "Extended Rail Controls")
-	bool bUseCustomPosition = true;
+	/* Use AbsolutePosition metadata to parameterize the spline*/
+	UPROPERTY(EditAnywhere, BlueprintReadWrite, Category = "Rail Controls")
+	bool bUseAbsolutePosition = true;
 
 	/* Custom parameter to drive current position*/
-	UPROPERTY(EditAnywhere, BlueprintReadWrite, Interp, Category = "Extended Rail Controls", meta=(EditCondition="bUseCustomPosition"))
-	float CustomPosition = 1.0f;
+	UPROPERTY(EditAnywhere, BlueprintReadWrite, Interp, Category = "Rail Controls", meta=(EditCondition="bUseAbsolutePosition"))
+	float AbsolutePositionOnRail = 1.0f;
 
-	/* Use CameraRotation metadata for attachment orientation*/
-	UPROPERTY(EditAnywhere, BlueprintReadWrite, Interp, Category = "Extended Rail Controls", meta = (EditCondition = "bLockOrientationToRail"))
-	bool bUseCameraRotationMetadata = true;
+	/* Use PointRotation metadata for attachment orientation. If false, attachment orientation is based on the spline curvature*/
+	UPROPERTY(EditAnywhere, BlueprintReadWrite, Interp, Category = "Rail Controls", meta = (EditCondition = "bLockOrientationToRail"))
+	bool bUsePointRotation = true;
 
 	/* Material assigned to spline component mesh*/
 	UPROPERTY(EditAnywhere, BlueprintSetter=SetSplineMeshMaterial, Category = "SplineVisualization")
 	TObjectPtr<UMaterialInterface> SplineMeshMaterial;
 
 	/* Material Instance Dynamic created for the spline mesh */
-	UPROPERTY(VisibleInstanceOnly, Transient, BlueprintReadOnly, Category = "SplineVisualization")
+	UPROPERTY(Transient, BlueprintReadOnly, Category = "SplineVisualization")
 	TArray<TObjectPtr<UMaterialInstanceDynamic>> SplineMeshMIDs;
 
 	/* Texture that can be set to SplineMeshMIDs */
@@ -46,40 +46,40 @@ public:
 	TObjectPtr<UTexture2D> SplineMeshTexture;
 
 	/* Determines if camera mount inherits LocationX*/
-	UPROPERTY(EditAnywhere, BlueprintReadWrite, Category = "Attachment")
+	UPROPERTY(EditAnywhere, BlueprintReadWrite, Category = "Attachment|Location")
 	bool bAttachLocationX = true;
 
 	/* Determines if camera mount inherits LocationY*/
-	UPROPERTY(EditAnywhere, BlueprintReadWrite, Category = "Attachment")
+	UPROPERTY(EditAnywhere, BlueprintReadWrite, Category = "Attachment|Location")
 	bool bAttachLocationY = true;
 
 	/* Determines if camera mount inherits LocationZ*/
-	UPROPERTY(EditAnywhere, BlueprintReadWrite, Category = "Attachment")
+	UPROPERTY(EditAnywhere, BlueprintReadWrite, Category = "Attachment|Location")
 	bool bAttachLocationZ = true;
 
 	/* Determines if camera mount inherits RotationX*/
-	UPROPERTY(EditAnywhere, BlueprintReadWrite, Category = "Attachment")
+	UPROPERTY(EditAnywhere, BlueprintReadWrite, Category = "Attachment|Rotation")
 	bool bAttachRotationX = true;
 
 	/* Determines if camera mount inherits RotationY*/
-	UPROPERTY(EditAnywhere, BlueprintReadWrite, Category = "Attachment")
+	UPROPERTY(EditAnywhere, BlueprintReadWrite, Category = "Attachment|Rotation")
 	bool bAttachRotationY = true;
 
 	/* Determines if camera mount inherits RotationZ*/
-	UPROPERTY(EditAnywhere, BlueprintReadWrite, Category = "Attachment")
+	UPROPERTY(EditAnywhere, BlueprintReadWrite, Category = "Attachment|Rotation")
 	bool bAttachRotationZ = true;
 
 	/* Determines if it can drive focal length on the attached actors*/
-	UPROPERTY(EditAnywhere, BlueprintReadWrite, Category = "Attachment")
-	bool bDriveFocalLength = true;
+	UPROPERTY(EditAnywhere, BlueprintReadWrite, Category = "Attachment|Camera")
+	bool bInheritFocalLength = true;
 
 	/* Determines if it can drive aperture on the attached actors*/
-	UPROPERTY(EditAnywhere, BlueprintReadWrite, Category = "Attachment")
-	bool bDriveAperture = true;
+	UPROPERTY(EditAnywhere, BlueprintReadWrite, Category = "Attachment|Camera")
+	bool bInheritAperture = true;
 
 	/* Determines if it can drive focus distance on the attached actors*/
-	UPROPERTY(EditAnywhere, BlueprintReadWrite, Category = "Attachment")
-	bool bDriveFocusDistance = true;
+	UPROPERTY(EditAnywhere, BlueprintReadWrite, Category = "Attachment|Camera")
+	bool bInheritFocusDistance = true;
 
 	/* Set spline mesh material*/
 	UFUNCTION(BlueprintSetter)
diff --git a/Engine/Plugins/Experimental/CineCameraRigs/Source/CineCameraRigs/Public/CineSplineComponent.h b/Engine/Plugins/Experimental/CineCameraRigs/Source/CineCameraRigs/Public/CineSplineComponent.h
--- a/Engine/Plugins/Experimental/CineCameraRigs/Source/CineCameraRigs/Public/CineSplineComponent.h
+++ b/Engine/Plugins/Experimental/CineCameraRigs/Source/CineCameraRigs/Public/CineSplineComponent.h
@@ -62,11 +62,11 @@ public:
 
 	/** Set custom position metadata at a given spline point*/
 	UFUNCTION(BlueprintCallable, Category = "CineSpline")
-	void SetCustomPositionAtSplinePoint(const int32 PointIndex, const float value);
+	void SetAbsolutePositionAtSplinePoint(const int32 PointIndex, const float value);
 
 	/* Set camera rotation metadata at a given spline point*/
 	UFUNCTION(BlueprintCallable, Category = "CineSpline")
-	void SetCameraRotationAtSplinePoint(const int32 PointIndex, const FQuat value);
+	void SetPointRotationAtSplinePoint(const int32 PointIndex, const FQuat value);
 
 	/** Returns true if there is a spline point at given custom position*/
 	UFUNCTION(BlueprintCallable, BlueprintPure, Category = "CineSpline")
@@ -94,11 +94,11 @@ public:
 
 	/* Get camera rotation metadata property value along the spline at spline point */
 	UFUNCTION(BlueprintCallable, Category = "CineSpline")
-	FQuat GetCameraRotationAtSplinePoint(int32 Index) const;
+	FQuat GetPointRotationAtSplinePoint(int32 Index) const;
 
 	/* Get camera rotation metadata value along the spline at spline input key*/
 	UFUNCTION(BlueprintCallable, Category = "CineSpline")
-	FQuat GetCameraRotationAtSplineInputKey(float InKey) const;
+	FQuat GetPointRotationAtSplineInputKey(float InKey) const;
 
 
 	DECLARE_MULTICAST_DELEGATE(FOnSplineEdited);
diff --git a/Engine/Plugins/Experimental/CineCameraRigs/Source/CineCameraRigs/Public/CineSplineMetadata.h b/Engine/Plugins/Experimental/CineCameraRigs/Source/CineCameraRigs/Public/CineSplineMetadata.h
--- a/Engine/Plugins/Experimental/CineCameraRigs/Source/CineCameraRigs/Public/CineSplineMetadata.h
+++ b/Engine/Plugins/Experimental/CineCameraRigs/Source/CineCameraRigs/Public/CineSplineMetadata.h
@@ -15,15 +15,15 @@ struct FCineSplineCurveDefaults
 	GENERATED_BODY()
 
 		FCineSplineCurveDefaults()
-		: DefaultCustomPosition(-1.0f)
+		: DefaultAbsolutePosition(-1.0f)
 		, DefaultFocalLength(35.0f)
 		, DefaultAperture(2.8f)
 		, DefaultFocusDistance(100000.f)
-		, DefaultCameraRotation(FQuat::Identity)
+		, DefaultPointRotation(FQuat::Identity)
 	{};
 
 	UPROPERTY(EditDefaultsOnly, Category = CineSplineCurveDefaults)
-	float DefaultCustomPosition;
+	float DefaultAbsolutePosition;
 
 	UPROPERTY(EditDefaultsOnly, Category = CineSplineCurveDefaults)
 	float DefaultFocalLength;
@@ -35,7 +35,7 @@ struct FCineSplineCurveDefaults
 	float DefaultFocusDistance;
 
 	UPROPERTY(EditDefaultsOnly, Category = CineSplineCurveDefaults)
-	FQuat DefaultCameraRotation;
+	FQuat DefaultPointRotation;
 };
 
 UCLASS()
@@ -56,7 +56,7 @@ public:
 	virtual void Fixup(int32 NumPoints, USplineComponent* SplineComp) override;
 
 	UPROPERTY(EditAnywhere, Category = "Point")
-	FInterpCurveFloat CustomPosition;
+	FInterpCurveFloat AbsolutePosition;
 
 	UPROPERTY(EditAnywhere, Category = "Camera")
 	FInterpCurveFloat FocalLength;
@@ -68,5 +68,5 @@ public:
 	FInterpCurveFloat FocusDistance;
 
 	UPROPERTY(EditAnywhere, Category = "Camera")
-	FInterpCurveQuat CameraRotation;
+	FInterpCurveQuat PointRotation;
 };
diff --git a/Engine/Plugins/Experimental/CineCameraRigs/Source/CineCameraRigsEditor/Private/CineCameraRigRailDetails.cpp b/Engine/Plugins/Experimental/CineCameraRigs/Source/CineCameraRigsEditor/Private/CineCameraRigRailDetails.cpp
--- a/Engine/Plugins/Experimental/CineCameraRigs/Source/CineCameraRigsEditor/Private/CineCameraRigRailDetails.cpp
+++ b/Engine/Plugins/Experimental/CineCameraRigs/Source/CineCameraRigsEditor/Private/CineCameraRigRailDetails.cpp
@@ -12,6 +12,7 @@
 #include "Widgets/Input/SNumericEntryBox.h"
 #include "Internationalization/Internationalization.h"
 #include "ScopedTransaction.h"
+#include "Layout/Visibility.h"
 
 #define LOCTEXT_NAMESPACE "FCineCameraRigRailDetails"
 
@@ -32,53 +33,63 @@ void FCineCameraRigRailDetails::CustomizeDetails(IDetailLayoutBuilder& DetailBui
 	}
 	RigRailActorPtr = RigRailActor;
 
-	TSharedRef<IPropertyHandle> CustomPositionPropertyHandle = DetailBuilder.GetProperty(GET_MEMBER_NAME_CHECKED(ACineCameraRigRail, CustomPosition));
-	TSharedRef<IPropertyHandle> UseCustomPositionPropertyHandle = DetailBuilder.GetProperty(GET_MEMBER_NAME_CHECKED(ACineCameraRigRail, bUseCustomPosition));
+	TSharedRef<IPropertyHandle> AbsolutePositionPropertyHandle = DetailBuilder.GetProperty(GET_MEMBER_NAME_CHECKED(ACineCameraRigRail, AbsolutePositionOnRail));
+	TSharedRef<IPropertyHandle> UseAbsolutePositionPropertyHandle = DetailBuilder.GetProperty(GET_MEMBER_NAME_CHECKED(ACineCameraRigRail, bUseAbsolutePosition));
 	TSharedRef<IPropertyHandle> OrigPositionPropertyHandle = DetailBuilder.GetProperty(GET_MEMBER_NAME_CHECKED(ACameraRig_Rail, CurrentPositionOnRail), ACameraRig_Rail::StaticClass());
+	TSharedRef<IPropertyHandle> LockOrientationPropertyHandle = DetailBuilder.GetProperty(GET_MEMBER_NAME_CHECKED(ACameraRig_Rail, bLockOrientationToRail), ACameraRig_Rail::StaticClass());
+	TSharedRef<IPropertyHandle> ShowRailPropertyHandle = DetailBuilder.GetProperty(GET_MEMBER_NAME_CHECKED(ACameraRig_Rail, bShowRailVisualization), ACameraRig_Rail::StaticClass());
+	TSharedRef<IPropertyHandle> PreviewScalePropertyHandle = DetailBuilder.GetProperty(GET_MEMBER_NAME_CHECKED(ACameraRig_Rail, PreviewMeshScale), ACameraRig_Rail::StaticClass());
 
-	const TAttribute<bool> EditCondition = TAttribute<bool>::Create([this, UseCustomPositionPropertyHandle]()
+
+
+	const TAttribute<bool> EditCondition = TAttribute<bool>::Create([this, UseAbsolutePositionPropertyHandle]()
 	{
 		bool bCond = false;
-		UseCustomPositionPropertyHandle->GetValue(bCond);
+		UseAbsolutePositionPropertyHandle->GetValue(bCond);
 		return !bCond;
 	});
 
-	DetailBuilder.EditDefaultProperty(OrigPositionPropertyHandle)->EditCondition(EditCondition, nullptr);
+	DetailBuilder.EditDefaultProperty(LockOrientationPropertyHandle)->Visibility(EVisibility::Hidden);
 
+	IDetailCategoryBuilder& RailControlsCategory = DetailBuilder.EditCategory("Rail Controls");
+	RailControlsCategory.AddProperty(UseAbsolutePositionPropertyHandle);
+	IDetailPropertyRow& AbsolutePositionRow = RailControlsCategory.AddProperty(AbsolutePositionPropertyHandle);
+	RailControlsCategory.AddProperty(OrigPositionPropertyHandle)
+	.EditCondition(EditCondition, nullptr);
 
-	if (IDetailPropertyRow* Row = DetailBuilder.EditDefaultProperty(CustomPositionPropertyHandle))
-	{
-		Row->CustomWidget()
+	AbsolutePositionRow.CustomWidget(false)
 		.NameContent()
 		[
-			CustomPositionPropertyHandle->CreatePropertyNameWidget()
+			AbsolutePositionPropertyHandle->CreatePropertyNameWidget()
 		]
 		.ValueContent()
 		.MinDesiredWidth(125.0f)
 		[
 			SNew(SNumericEntryBox<float>)
-			.ToolTipText(LOCTEXT("CurrentPostionToolTip","Postion property using custom parameterization"))
+			.ToolTipText(LOCTEXT("CurrentPostionToolTip", "Postion property using custom parameterization"))
 			.AllowSpin(true)
-			.MinSliderValue(this, &FCineCameraRigRailDetails::GetCustomPositionSliderMinValue)
-			.MaxSliderValue(this, &FCineCameraRigRailDetails::GetCustomPositionSliderMaxValue)
-			.Value(this, &FCineCameraRigRailDetails::GetCustomPosition)
-			.OnValueChanged(this, &FCineCameraRigRailDetails::OnCustomPositionChanged)
-			.OnValueCommitted(this, &FCineCameraRigRailDetails::OnCustomPositionCommitted)
-			.OnBeginSliderMovement(this, &FCineCameraRigRailDetails::OnBeginCustomPositionSliderMovement)
-			.OnEndSliderMovement(this, &FCineCameraRigRailDetails::OnEndCustomPositionSliderMovement)
+			.MinSliderValue(this, &FCineCameraRigRailDetails::GetAbsolutePositionSliderMinValue)
+			.MaxSliderValue(this, &FCineCameraRigRailDetails::GetAbsolutePositionSliderMaxValue)
+			.Value(this, &FCineCameraRigRailDetails::GetAbsolutePosition)
+			.OnValueChanged(this, &FCineCameraRigRailDetails::OnAbsolutePositionChanged)
+			.OnValueCommitted(this, &FCineCameraRigRailDetails::OnAbsolutePositionCommitted)
+			.OnBeginSliderMovement(this, &FCineCameraRigRailDetails::OnBeginAbsolutePositionSliderMovement)
+			.OnEndSliderMovement(this, &FCineCameraRigRailDetails::OnEndAbsolutePositionSliderMovement)
 			.Font(IDetailLayoutBuilder::GetDetailFont())
 		];
-	}
 
+	IDetailCategoryBuilder& VisualizationCategory = DetailBuilder.EditCategory("SplineVisualization");
+	VisualizationCategory.AddProperty(ShowRailPropertyHandle);
+	VisualizationCategory.AddProperty(PreviewScalePropertyHandle);
 }
 
-void FCineCameraRigRailDetails::OnCustomPositionChanged(float NewValue)
+void FCineCameraRigRailDetails::OnAbsolutePositionChanged(float NewValue)
 {
 	if (ACineCameraRigRail* RigRailActor = RigRailActorPtr.Get())
 	{
-		RigRailActor->CustomPosition = NewValue;
-		static FProperty* CustomPositionProperty = FindFProperty<FProperty>(ACineCameraRigRail::StaticClass(), GET_MEMBER_NAME_CHECKED(ACineCameraRigRail, CustomPosition));
-		FPropertyChangedEvent SetValueEvent(CustomPositionProperty);
+		RigRailActor->AbsolutePositionOnRail = NewValue;
+		static FProperty* AbsolutePositionProperty = FindFProperty<FProperty>(ACineCameraRigRail::StaticClass(), GET_MEMBER_NAME_CHECKED(ACineCameraRigRail, AbsolutePositionOnRail));
+		FPropertyChangedEvent SetValueEvent(AbsolutePositionProperty);
 		RigRailActor->PostEditChangeProperty(SetValueEvent);
 		if (GEditor)
 		{
@@ -87,16 +98,16 @@ void FCineCameraRigRailDetails::OnCustomPositionChanged(float NewValue)
 	}
 }
 
-void FCineCameraRigRailDetails::OnCustomPositionCommitted(float NewValue, ETextCommit::Type CommitType)
+void FCineCameraRigRailDetails::OnAbsolutePositionCommitted(float NewValue, ETextCommit::Type CommitType)
 {
 	if (ACineCameraRigRail* RigRailActor = RigRailActorPtr.Get())
 	{
-		const FScopedTransaction Transaction(LOCTEXT("SetCustomPosition", "Set rig rail custom position"));
+		const FScopedTransaction Transaction(LOCTEXT("SetAbsolutePosition", "Set rig rail custom position"));
 		RigRailActor->SetFlags(RF_Transactional);
 		RigRailActor->Modify();
-		RigRailActor->CustomPosition = NewValue;
-		static FProperty* CustomPositionProperty = FindFProperty<FProperty>(ACineCameraRigRail::StaticClass(), GET_MEMBER_NAME_CHECKED(ACineCameraRigRail, CustomPosition));
-		FPropertyChangedEvent SetValueEvent(CustomPositionProperty);
+		RigRailActor->AbsolutePositionOnRail = NewValue;
+		static FProperty* AbsolutePositionProperty = FindFProperty<FProperty>(ACineCameraRigRail::StaticClass(), GET_MEMBER_NAME_CHECKED(ACineCameraRigRail, AbsolutePositionOnRail));
+		FPropertyChangedEvent SetValueEvent(AbsolutePositionProperty);
 		RigRailActor->PostEditChangeProperty(SetValueEvent);
 		if (GEditor)
 		{
@@ -105,14 +116,14 @@ void FCineCameraRigRailDetails::OnCustomPositionCommitted(float NewValue, ETextC
 	}
 }
 
-void FCineCameraRigRailDetails::OnBeginCustomPositionSliderMovement()
+void FCineCameraRigRailDetails::OnBeginAbsolutePositionSliderMovement()
 {
-	if (!bCustomPositionSliderStartedTransaction)
+	if (!bAbsolutePositionSliderStartedTransaction)
 	{
 		if (GEditor)
 		{
-			bCustomPositionSliderStartedTransaction = true;
-			GEditor->BeginTransaction(LOCTEXT("CustomPositionSliderTransaction", "Set rig rail custom position via slider"));
+			bAbsolutePositionSliderStartedTransaction = true;
+			GEditor->BeginTransaction(LOCTEXT("AbsolutePositionSliderTransaction", "Set rig rail custom position via slider"));
 			if (ACineCameraRigRail* RigRailActor = RigRailActorPtr.Get())
 			{
 				RigRailActor->SetFlags(RF_Transactional);
@@ -122,28 +133,28 @@ void FCineCameraRigRailDetails::OnBeginCustomPositionSliderMovement()
 	}
 }
 
-void FCineCameraRigRailDetails::OnEndCustomPositionSliderMovement(float NewValue)
+void FCineCameraRigRailDetails::OnEndAbsolutePositionSliderMovement(float NewValue)
 {
-	if (bCustomPositionSliderStartedTransaction)
+	if (bAbsolutePositionSliderStartedTransaction)
 	{
 		if (GEditor)
 		{
 			GEditor->EndTransaction();
-			bCustomPositionSliderStartedTransaction = false;
+			bAbsolutePositionSliderStartedTransaction = false;
 		}
 	}
 }
 
-TOptional<float> FCineCameraRigRailDetails::GetCustomPosition() const
+TOptional<float> FCineCameraRigRailDetails::GetAbsolutePosition() const
 {
 	if (ACineCameraRigRail* RigRailActor = RigRailActorPtr.Get())
 	{
-		return RigRailActor->CustomPosition;
+		return RigRailActor->AbsolutePositionOnRail;
 	}
 	return 0.0f;
 }
 
-TOptional<float> FCineCameraRigRailDetails::GetCustomPositionSliderMinValue() const
+TOptional<float> FCineCameraRigRailDetails::GetAbsolutePositionSliderMinValue() const
 {
 
 	float MinValue = 1.0f;
@@ -151,12 +162,12 @@ TOptional<float> FCineCameraRigRailDetails::GetCustomPositionSliderMinValue() co
 	{
 		UCineSplineComponent* SplineComp = RigRailActor->GetCineSplineComponent();
 		const UCineSplineMetadata* MetaData = Cast<UCineSplineMetadata>(SplineComp->GetSplinePointsMetadata());
-		MinValue = MetaData->CustomPosition.Points[0].OutVal;
+		MinValue = MetaData->AbsolutePosition.Points[0].OutVal;
 	}
 	return MinValue;
 }
 
-TOptional<float> FCineCameraRigRailDetails::GetCustomPositionSliderMaxValue() const
+TOptional<float> FCineCameraRigRailDetails::GetAbsolutePositionSliderMaxValue() const
 {
 
 	float MaxValue = 5.0f;
@@ -164,8 +175,8 @@ TOptional<float> FCineCameraRigRailDetails::GetCustomPositionSliderMaxValue() co
 	{
 		UCineSplineComponent* SplineComp = RigRailActor->GetCineSplineComponent();
 		const UCineSplineMetadata* MetaData = Cast<UCineSplineMetadata>(SplineComp->GetSplinePointsMetadata());
-		int32 NumPoints = MetaData->CustomPosition.Points.Num();
-		MaxValue = MetaData->CustomPosition.Points[NumPoints - 1].OutVal;
+		int32 NumPoints = MetaData->AbsolutePosition.Points.Num();
+		MaxValue = MetaData->AbsolutePosition.Points[NumPoints - 1].OutVal;
 	}
 	return MaxValue;
 }
diff --git a/Engine/Plugins/Experimental/CineCameraRigs/Source/CineCameraRigsEditor/Private/CineCameraRigRailDetails.h b/Engine/Plugins/Experimental/CineCameraRigs/Source/CineCameraRigsEditor/Private/CineCameraRigRailDetails.h
--- a/Engine/Plugins/Experimental/CineCameraRigs/Source/CineCameraRigsEditor/Private/CineCameraRigRailDetails.h
+++ b/Engine/Plugins/Experimental/CineCameraRigs/Source/CineCameraRigsEditor/Private/CineCameraRigRailDetails.h
@@ -26,14 +26,15 @@ public:
 private:
 	TWeakObjectPtr<ACineCameraRigRail> RigRailActorPtr;
 
-	TOptional<float> GetCustomPosition() const;
+	TOptional<float> GetAbsolutePosition() const;
 
-	TOptional<float> GetCustomPositionSliderMinValue() const;
-	TOptional<float> GetCustomPositionSliderMaxValue() const;
-	void OnBeginCustomPositionSliderMovement();
-	void OnEndCustomPositionSliderMovement(float NewValue);
-	void OnCustomPositionCommitted(float NewValue, ETextCommit::Type CommitType);
-	void OnCustomPositionChanged(float NewValue);
+	TOptional<float> GetAbsolutePositionSliderMinValue() const;
+	TOptional<float> GetAbsolutePositionSliderMaxValue() const;
+	void OnBeginAbsolutePositionSliderMovement();
+	void OnEndAbsolutePositionSliderMovement(float NewValue);
+	void OnAbsolutePositionCommitted(float NewValue, ETextCommit::Type CommitType);
+	void OnAbsolutePositionChanged(float NewValue);
+
+	bool bAbsolutePositionSliderStartedTransaction = false;
 
-	bool bCustomPositionSliderStartedTransaction = false;
 };
diff --git a/Engine/Plugins/Experimental/CineCameraRigs/Source/CineCameraRigsEditor/Private/CineSplineComponentVisualizer.cpp b/Engine/Plugins/Experimental/CineCameraRigs/Source/CineCameraRigsEditor/Private/CineSplineComponentVisualizer.cpp
--- a/Engine/Plugins/Experimental/CineCameraRigs/Source/CineCameraRigsEditor/Private/CineSplineComponentVisualizer.cpp
+++ b/Engine/Plugins/Experimental/CineCameraRigs/Source/CineCameraRigsEditor/Private/CineSplineComponentVisualizer.cpp
@@ -45,8 +45,8 @@ void FCineSplineComponentVisualizer::DrawVisualizationHUD(const UActorComponent*
 				FNumberFormattingOptions FmtOptions;
 				FmtOptions.SetMaximumFractionalDigits(4);
 				FmtOptions.SetMinimumFractionalDigits(1);
-				const FText Text = FText::AsNumber(Metadata->CustomPosition.Points[i].OutVal, &FmtOptions);
-				FLinearColor Color = FMath::IsNearlyEqual(FMath::Frac(Metadata->CustomPosition.Points[i].OutVal), 0.0f) ? FLinearColor(1,1,0,1) : FLinearColor(0, 1, 1, 1);
+				const FText Text = FText::AsNumber(Metadata->AbsolutePosition.Points[i].OutVal, &FmtOptions);
+				FLinearColor Color = FMath::IsNearlyEqual(FMath::Frac(Metadata->AbsolutePosition.Points[i].OutVal), 0.0f) ? FLinearColor(1,1,0,1) : FLinearColor(0, 1, 1, 1);
 				Canvas->DrawShadowedString(DrawPositionX, DrawPositionY, *Text.ToString(), GEngine->GetLargeFont(), Color, FLinearColor::Black);
 			}
 		}
diff --git a/Engine/Plugins/Experimental/CineCameraRigs/Source/CineCameraRigsEditor/Private/CineSplineMetadataDetails.cpp b/Engine/Plugins/Experimental/CineCameraRigs/Source/CineCameraRigsEditor/Private/CineSplineMetadataDetails.cpp
--- a/Engine/Plugins/Experimental/CineCameraRigs/Source/CineCameraRigsEditor/Private/CineSplineMetadataDetails.cpp
+++ b/Engine/Plugins/Experimental/CineCameraRigs/Source/CineCameraRigsEditor/Private/CineSplineMetadataDetails.cpp
@@ -30,7 +30,7 @@ TSharedPtr<ISplineMetadataDetails> UCineSplineMetadataDetailsFactory::Create()
 
 FText FCineSplineMetadataDetails::GetDisplayName() const
 {
-	return LOCTEXT("CineSplineMetadataDetails", "CineSpline");
+	return LOCTEXT("CineSplineMetadataDetails", "CineSplinePointData");
 }
 
 template<class T>
@@ -53,11 +53,11 @@ void FCineSplineMetadataDetails::Update(USplineComponent* InSplineComponent, con
 {
 	SplineComp = InSplineComponent;
 	SelectedKeys = InSelectedKeys;
-	CustomPositionValue.Reset();
+	AbsolutePositionValue.Reset();
 	FocalLengthValue.Reset();
 	ApertureValue.Reset();
 	FocusDistanceValue.Reset();
-	CameraRotationValue.Reset();
+	PointRotationValue.Reset();
 
 	if (!IsValid(InSplineComponent))
 	{
@@ -69,19 +69,19 @@ void FCineSplineMetadataDetails::Update(USplineComponent* InSplineComponent, con
 		return;
 	}
 
-	bool bUpdateCustomPosition = true;
+	bool bUpdateAbsolutePosition = true;
 	bool bUpdateFocalLength = true;
 	bool bUpdateAperture = true;
 	bool bUpdateFocusDistance = true;
-	bool bUpdateCameraRotation = true;
+	bool bUpdatePointRotation = true;
 
 	for (const int32 Index : InSelectedKeys)
 	{
-		bUpdateCustomPosition = bUpdateCustomPosition ? UpdateMultipleValue(CustomPositionValue, Metadata->CustomPosition.Points[Index].OutVal) : false;
+		bUpdateAbsolutePosition = bUpdateAbsolutePosition ? UpdateMultipleValue(AbsolutePositionValue, Metadata->AbsolutePosition.Points[Index].OutVal) : false;
 		bUpdateFocalLength = bUpdateFocalLength ? UpdateMultipleValue(FocalLengthValue, Metadata->FocalLength.Points[Index].OutVal) : false;
 		bUpdateAperture = bUpdateAperture ? UpdateMultipleValue(ApertureValue, Metadata->Aperture.Points[Index].OutVal) : false;
 		bUpdateFocusDistance = bUpdateFocusDistance ? UpdateMultipleValue(FocusDistanceValue, Metadata->FocusDistance.Points[Index].OutVal) : false;
-		bUpdateCameraRotation = bUpdateCameraRotation ? UpdateMultipleValue(CameraRotationValue, Metadata->CameraRotation.Points[Index].OutVal) : false;
+		bUpdatePointRotation = bUpdatePointRotation ? UpdateMultipleValue(PointRotationValue, Metadata->PointRotation.Points[Index].OutVal) : false;
 	}
 }
 
@@ -121,12 +121,12 @@ UCineSplineMetadata* FCineSplineMetadataDetails::GetMetadata() const
 	return SplineComp ? Cast<UCineSplineMetadata>(SplineComp->GetSplinePointsMetadata()) : nullptr;
 }
 
-void FCineSplineMetadataDetails::OnSetCustomPosition(float NewValue, ETextCommit::Type CommitInfo)
+void FCineSplineMetadataDetails::OnSetAbsolutePosition(float NewValue, ETextCommit::Type CommitInfo)
 {
 	if (UCineSplineMetadata* Metadata = GetMetadata())
 	{
-		const FScopedTransaction Transaction(LOCTEXT("SetCustomPosition", "Set spline point custom position data"));
-		SetValues<float>(*this, Metadata->CustomPosition.Points, NewValue);
+		const FScopedTransaction Transaction(LOCTEXT("SetAbsolutePosition", "Set spline point custom position data"));
+		SetValues<float>(*this, Metadata->AbsolutePosition.Points, NewValue);
 	}
 }
 
@@ -159,35 +159,36 @@ void FCineSplineMetadataDetails::OnSetFocusDistance(float NewValue, ETextCommit:
 
 TOptional<float> FCineSplineMetadataDetails::GetRotation(EAxis::Type Axis) const
 {
-	return static_cast<float>(CameraRotationValue.Get(FQuat::Identity).Rotator().GetComponentForAxis(Axis));
+	return static_cast<float>(PointRotationValue.Get(FQuat::Identity).Rotator().GetComponentForAxis(Axis));
 }
 void FCineSplineMetadataDetails::OnSetRotation(float NewValue, ETextCommit::Type CommitInfo, EAxis::Type Axis)
 {
 	if (UCineSplineMetadata* Metadata = GetMetadata())
 	{
-		const FScopedTransaction Transaction(LOCTEXT("SetCameraRotation", "Set spline point camera rotation data"));
-		FRotator CurrentRotator(CameraRotationValue.Get(FQuat::Identity));
+		const FScopedTransaction Transaction(LOCTEXT("SetPointRotation", "Set spline point camera rotation data"));
+		FRotator CurrentRotator(PointRotationValue.Get(FQuat::Identity));
 		switch (Axis)
 		{
 		case EAxis::X: CurrentRotator.Roll = NewValue; break;
 		case EAxis::Y: CurrentRotator.Pitch = NewValue; break;
 		case EAxis::Z: CurrentRotator.Yaw = NewValue; break;
 		}
-		SetValues<FQuat>(*this, Metadata->CameraRotation.Points, CurrentRotator.Quaternion());
+		SetValues<FQuat>(*this, Metadata->PointRotation.Points, CurrentRotator.Quaternion());
 	}
 }
 
 void FCineSplineMetadataDetails::GenerateChildContent(IDetailGroup& DetailGroup)
 {
 	DetailGroup.AddWidgetRow()
-		.RowTag("CustomPosition")
+		.ShouldAutoExpand(true)
+		.RowTag("AbsolutePosition")
 		.Visibility(TAttribute<EVisibility>(this, &FCineSplineMetadataDetails::IsEnabled))
 		.NameContent()
 		.HAlign(HAlign_Left)
 		.VAlign(VAlign_Center)
 		[
 			SNew(STextBlock)
-			.Text(LOCTEXT("CustomPosition", "Custom Position"))
+			.Text(LOCTEXT("AbsolutePosition", "Absolute Position"))
 			.Font(IDetailLayoutBuilder::GetDetailFont())
 		]
 		.ValueContent()
@@ -195,7 +196,7 @@ void FCineSplineMetadataDetails::GenerateChildContent(IDetailGroup& DetailGroup)
 		.MaxDesiredWidth(125.0f)
 		[
 			SNew(SNumericEntryBox<float>)
-			.Value(this, &FCineSplineMetadataDetails::GetCustomPosition)
+			.Value(this, &FCineSplineMetadataDetails::GetAbsolutePosition)
 			.AllowSpin(true)
 			.MinValue(TOptional<float>())
 			.MaxValue(TOptional<float>())
@@ -204,8 +205,38 @@ void FCineSplineMetadataDetails::GenerateChildContent(IDetailGroup& DetailGroup)
 			.OnBeginSliderMovement(this, &FCineSplineMetadataDetails::OnBeginSliderMovement)
 			.OnEndSliderMovement(this, &FCineSplineMetadataDetails::OnEndSliderMovement)
 			.UndeterminedString(LOCTEXT("Multiple", "Multiple"))
-			.OnValueCommitted(this, &FCineSplineMetadataDetails::OnSetCustomPosition)
-			.OnValueChanged(this, &FCineSplineMetadataDetails::OnSetCustomPosition, ETextCommit::Default)
+			.OnValueCommitted(this, &FCineSplineMetadataDetails::OnSetAbsolutePosition)
+			.OnValueChanged(this, &FCineSplineMetadataDetails::OnSetAbsolutePosition, ETextCommit::Default)
+			.Font(IDetailLayoutBuilder::GetDetailFont())
+		];
+
+	DetailGroup.AddWidgetRow()
+		.RowTag("PointRotation")
+		.Visibility(TAttribute<EVisibility>(this, &FCineSplineMetadataDetails::IsEnabled))
+		.NameContent()
+		.HAlign(HAlign_Left)
+		.VAlign(VAlign_Center)
+		[
+			SNew(STextBlock)
+			.Text(LOCTEXT("PointRotation", "Point Rotation"))
+			.Font(IDetailLayoutBuilder::GetDetailFont())
+		]
+		.ValueContent()
+		.MinDesiredWidth(125.0f)
+		.MaxDesiredWidth(125.0f)
+		[
+			SNew(SRotatorInputBox)
+			.Roll(this, &FCineSplineMetadataDetails::GetRotation, EAxis::X)
+			.Pitch(this, &FCineSplineMetadataDetails::GetRotation, EAxis::Y)
+			.Yaw(this, &FCineSplineMetadataDetails::GetRotation, EAxis::Z)
+			.OnRollChanged(this, &FCineSplineMetadataDetails::OnSetRotation, ETextCommit::Default, EAxis::X)
+			.OnPitchChanged(this, &FCineSplineMetadataDetails::OnSetRotation, ETextCommit::Default, EAxis::Y)
+			.OnYawChanged(this, &FCineSplineMetadataDetails::OnSetRotation, ETextCommit::Default, EAxis::Z)
+			.OnRollCommitted(this, &FCineSplineMetadataDetails::OnSetRotation, EAxis::X)
+			.OnPitchCommitted(this, &FCineSplineMetadataDetails::OnSetRotation, EAxis::Y)
+			.OnYawCommitted(this, &FCineSplineMetadataDetails::OnSetRotation, EAxis::Z)
+			.AllowSpin(true)
+			.bColorAxisLabels(false)
 			.Font(IDetailLayoutBuilder::GetDetailFont())
 		];
 
@@ -301,35 +332,7 @@ void FCineSplineMetadataDetails::GenerateChildContent(IDetailGroup& DetailGroup)
 			.Font(IDetailLayoutBuilder::GetDetailFont())
 		];
 
-		DetailGroup.AddWidgetRow()
-			.RowTag("CameraRotation")
-			.Visibility(TAttribute<EVisibility>(this, &FCineSplineMetadataDetails::IsEnabled))
-			.NameContent()
-			.HAlign(HAlign_Left)
-			.VAlign(VAlign_Center)
-			[
-				SNew(STextBlock)
-				.Text(LOCTEXT("CameraRotation", "Camera Rotation"))
-				.Font(IDetailLayoutBuilder::GetDetailFont())
-			]
-			.ValueContent()
-			.MinDesiredWidth(125.0f)
-			.MaxDesiredWidth(125.0f)
-			[
-				SNew(SRotatorInputBox)
-				.Roll(this, &FCineSplineMetadataDetails::GetRotation, EAxis::X)
-				.Pitch(this, &FCineSplineMetadataDetails::GetRotation, EAxis::Y)
-				.Yaw(this, &FCineSplineMetadataDetails::GetRotation, EAxis::Z)
-				.OnRollChanged(this, &FCineSplineMetadataDetails::OnSetRotation, ETextCommit::Default, EAxis::X)
-				.OnPitchChanged(this, &FCineSplineMetadataDetails::OnSetRotation, ETextCommit::Default, EAxis::Y)
-				.OnYawChanged(this, &FCineSplineMetadataDetails::OnSetRotation, ETextCommit::Default, EAxis::Z)
-				.OnRollCommitted(this, &FCineSplineMetadataDetails::OnSetRotation, EAxis::X)
-				.OnPitchCommitted(this, &FCineSplineMetadataDetails::OnSetRotation, EAxis::Y)
-				.OnYawCommitted(this, &FCineSplineMetadataDetails::OnSetRotation, EAxis::Z)
-				.AllowSpin(false)
-				.bColorAxisLabels(false)
-				.Font(IDetailLayoutBuilder::GetDetailFont())
-			];
+
 }
 
 #undef LOCTEXT_NAMESPACE
\ No newline at end of file
diff --git a/Engine/Plugins/Experimental/CineCameraRigs/Source/CineCameraRigsEditor/Private/CineSplineMetadataDetails.h b/Engine/Plugins/Experimental/CineCameraRigs/Source/CineCameraRigsEditor/Private/CineSplineMetadataDetails.h
--- a/Engine/Plugins/Experimental/CineCameraRigs/Source/CineCameraRigsEditor/Private/CineSplineMetadataDetails.h
+++ b/Engine/Plugins/Experimental/CineCameraRigs/Source/CineCameraRigsEditor/Private/CineSplineMetadataDetails.h
@@ -33,11 +33,11 @@ public:
 	virtual void Update(USplineComponent* InSplineComponent, const TSet<int32>& InSelectedKeys) override;
 	virtual void GenerateChildContent(IDetailGroup& InGroup) override;
 
-	TOptional<float> CustomPositionValue;
+	TOptional<float> AbsolutePositionValue;
 	TOptional<float> FocalLengthValue;
 	TOptional<float> ApertureValue;
 	TOptional<float> FocusDistanceValue;
-	TOptional<FQuat> CameraRotationValue;
+	TOptional<FQuat> PointRotationValue;
 
 	USplineComponent* SplineComp = nullptr;
 	TSet<int32> SelectedKeys;
@@ -50,8 +50,8 @@ private:
 	bool IsOnePointSelected() const { return SelectedKeys.Num() == 1; }
 
 
-	TOptional<float> GetCustomPosition() const { return CustomPositionValue; }
-	void OnSetCustomPosition(float NewValue, ETextCommit::Type CommitInfo);
+	TOptional<float> GetAbsolutePosition() const { return AbsolutePositionValue; }
+	void OnSetAbsolutePosition(float NewValue, ETextCommit::Type CommitInfo);
 
 	TOptional<float> GetFocalLength() const { return FocalLengthValue; }
 	void OnSetFocalLength(float NewValue, ETextCommit::Type CommitInfo);
