Update VCam spline mode to work with MultiUser

diff --git a/Engine/Plugins/Experimental/CineCameraRigs/CineCameraRigs.uplugin b/Engine/Plugins/Experimental/CineCameraRigs/CineCameraRigs.uplugin
--- a/Engine/Plugins/Experimental/CineCameraRigs/CineCameraRigs.uplugin
+++ b/Engine/Plugins/Experimental/CineCameraRigs/CineCameraRigs.uplugin
@@ -31,5 +31,9 @@
 			"Name": "EditorScriptingUtilities",
 			"Enabled": true
 		},
+		{
+			"Name": "ConcertSyncCore",
+			"Enabled": true
+		}
 	]
 }
diff --git a/Engine/Plugins/Experimental/CineCameraRigs/Source/CineCameraRigs/CineCameraRigs.Build.cs b/Engine/Plugins/Experimental/CineCameraRigs/Source/CineCameraRigs/CineCameraRigs.Build.cs
--- a/Engine/Plugins/Experimental/CineCameraRigs/Source/CineCameraRigs/CineCameraRigs.Build.cs
+++ b/Engine/Plugins/Experimental/CineCameraRigs/Source/CineCameraRigs/CineCameraRigs.Build.cs
@@ -42,6 +42,7 @@ public class CineCameraRigs : ModuleRules
 				"MovieSceneTracks",
 				"Slate",
 				"SlateCore",
+				"ConcertSyncCore",
 				// ... add private dependencies that you statically link with here ...	
 			}
 			);
diff --git a/Engine/Plugins/Experimental/CineCameraRigs/Source/CineCameraRigs/Private/CineCameraRigs.cpp b/Engine/Plugins/Experimental/CineCameraRigs/Source/CineCameraRigs/Private/CineCameraRigs.cpp
--- a/Engine/Plugins/Experimental/CineCameraRigs/Source/CineCameraRigs/Private/CineCameraRigs.cpp
+++ b/Engine/Plugins/Experimental/CineCameraRigs/Source/CineCameraRigs/Private/CineCameraRigs.cpp
@@ -2,6 +2,7 @@
 
 #include "CineCameraRigs.h"
 #include "CineSplineLog.h"
+#include "ConcertSyncSettings.h"
 
 #define LOCTEXT_NAMESPACE "FCineCameraRigsModule"
 
@@ -10,6 +11,16 @@ DEFINE_LOG_CATEGORY(LogCineSpline);
 void FCineCameraRigsModule::StartupModule()
 {
 	// This code will execute after your module is loaded into memory; the exact timing is specified in the .uplugin file per-module
+
+	// Update MU filter
+	UConcertSyncConfig* SyncConfig = GetMutableDefault<UConcertSyncConfig>();
+	if(SyncConfig)
+	{ 
+		FTransactionClassFilter Filter;
+		Filter.ObjectOuterClass = FSoftClassPath(TEXT("/Script/Engine.World"));
+		Filter.ObjectClasses.Add(FSoftClassPath(TEXT("/Script/CineCameraRigs.CineSplineMetadata")));
+		SyncConfig->IncludeObjectClassFilters.Add(Filter);
+	}
 }
 
 void FCineCameraRigsModule::ShutdownModule()
diff --git a/Engine/Plugins/Experimental/CineCameraRigs/Source/CineCameraRigsEditor/Private/CineSplineMetadataDetails.cpp b/Engine/Plugins/Experimental/CineCameraRigs/Source/CineCameraRigsEditor/Private/CineSplineMetadataDetails.cpp
--- a/Engine/Plugins/Experimental/CineCameraRigs/Source/CineCameraRigsEditor/Private/CineSplineMetadataDetails.cpp
+++ b/Engine/Plugins/Experimental/CineCameraRigs/Source/CineCameraRigsEditor/Private/CineSplineMetadataDetails.cpp
@@ -88,6 +88,7 @@ void FCineSplineMetadataDetails::Update(USplineComponent* InSplineComponent, con
 template<class T>
 void SetValues(FCineSplineMetadataDetails& Details, TArray<FInterpCurvePoint<T>>& Points, const T& NewValue)
 {
+	const FScopedTransaction Transaction(LOCTEXT("EditSplinePointsMetadata", "Edit spline point metadata"), (Details.SliderEnterCount == 0) );
 	Details.SplineComp->GetSplinePointsMetadata()->Modify();
 	for (int32 Index : Details.SelectedKeys)
 	{
@@ -108,12 +109,22 @@ void SetValues(FCineSplineMetadataDetails& Details, TArray<FInterpCurvePoint<T>>
 
 void FCineSplineMetadataDetails::OnBeginSliderMovement()
 {
-	EditSliderValueTransaction = MakeUnique<FScopedTransaction>(LOCTEXT("EditCineSplineProperty", "Edit CineSpline Property"));
+	if (SliderEnterCount == 0)
+	{
+		GEditor->BeginTransaction(LOCTEXT("EditCineSplineProperty", "Edit CineSpline Property"));
+	}
+	SliderEnterCount++;
+	SplineComp->Modify();
 }
 
 void FCineSplineMetadataDetails::OnEndSliderMovement(float NewValue)
 {
-	EditSliderValueTransaction.Reset();
+	SliderEnterCount--;
+	check(SliderEnterCount >= 0);
+	if (SliderEnterCount == 0)
+	{
+		GEditor->EndTransaction();
+	}
 }
 
 UCineSplineMetadata* FCineSplineMetadataDetails::GetMetadata() const
@@ -125,7 +136,6 @@ void FCineSplineMetadataDetails::OnSetAbsolutePosition(float NewValue, ETextComm
 {
 	if (UCineSplineMetadata* Metadata = GetMetadata())
 	{
-		const FScopedTransaction Transaction(LOCTEXT("SetAbsolutePosition", "Set spline point custom position data"));
 		SetValues<float>(*this, Metadata->AbsolutePosition.Points, NewValue);
 	}
 }
@@ -134,7 +144,6 @@ void FCineSplineMetadataDetails::OnSetFocalLength(float NewValue, ETextCommit::T
 {
 	if (UCineSplineMetadata* Metadata = GetMetadata())
 	{
-		const FScopedTransaction Transaction(LOCTEXT("SetFocalLength", "Set spline point focal length data"));
 		SetValues<float>(*this, Metadata->FocalLength.Points, NewValue);
 	}
 }
@@ -143,7 +152,6 @@ void FCineSplineMetadataDetails::OnSetAperture(float NewValue, ETextCommit::Type
 {
 	if (UCineSplineMetadata* Metadata = GetMetadata())
 	{
-		const FScopedTransaction Transaction(LOCTEXT("SetAperture", "Set spline point aperture data"));
 		SetValues<float>(*this, Metadata->Aperture.Points, NewValue);
 	}
 }
@@ -152,7 +160,6 @@ void FCineSplineMetadataDetails::OnSetFocusDistance(float NewValue, ETextCommit:
 {
 	if (UCineSplineMetadata* Metadata = GetMetadata())
 	{
-		const FScopedTransaction Transaction(LOCTEXT("SetFocusDistance", "Set spline point focus distance data"));
 		SetValues<float>(*this, Metadata->FocusDistance.Points, NewValue);
 	}
 }
@@ -165,7 +172,6 @@ void FCineSplineMetadataDetails::OnSetRotation(float NewValue, ETextCommit::Type
 {
 	if (UCineSplineMetadata* Metadata = GetMetadata())
 	{
-		const FScopedTransaction Transaction(LOCTEXT("SetPointRotation", "Set spline point camera rotation data"));
 		FRotator CurrentRotator(PointRotationValue.Get(FQuat::Identity));
 		switch (Axis)
 		{
diff --git a/Engine/Plugins/Experimental/CineCameraRigs/Source/CineCameraRigsEditor/Private/CineSplineMetadataDetails.h b/Engine/Plugins/Experimental/CineCameraRigs/Source/CineCameraRigsEditor/Private/CineSplineMetadataDetails.h
--- a/Engine/Plugins/Experimental/CineCameraRigs/Source/CineCameraRigsEditor/Private/CineSplineMetadataDetails.h
+++ b/Engine/Plugins/Experimental/CineCameraRigs/Source/CineCameraRigsEditor/Private/CineSplineMetadataDetails.h
@@ -41,6 +41,7 @@ public:
 
 	USplineComponent* SplineComp = nullptr;
 	TSet<int32> SelectedKeys;
+	int32 SliderEnterCount = 0;
 
 private:
 	void OnBeginSliderMovement();
