Fix for duplicate MU filter entry

diff --git a/Engine/Plugins/Experimental/CineCameraRigs/Source/CineCameraRigs/Private/CineCameraRigs.cpp b/Engine/Plugins/Experimental/CineCameraRigs/Source/CineCameraRigs/Private/CineCameraRigs.cpp
--- a/Engine/Plugins/Experimental/CineCameraRigs/Source/CineCameraRigs/Private/CineCameraRigs.cpp
+++ b/Engine/Plugins/Experimental/CineCameraRigs/Source/CineCameraRigs/Private/CineCameraRigs.cpp
@@ -3,6 +3,7 @@
 #include "CineCameraRigs.h"
 #include "CineSplineLog.h"
 #include "ConcertSyncSettings.h"
+#include "Algo/AnyOf.h"
 
 #define LOCTEXT_NAMESPACE "FCineCameraRigsModule"
 
@@ -16,10 +17,20 @@ void FCineCameraRigsModule::StartupModule()
 	UConcertSyncConfig* SyncConfig = GetMutableDefault<UConcertSyncConfig>();
 	if(SyncConfig)
 	{ 
-		FTransactionClassFilter Filter;
-		Filter.ObjectOuterClass = FSoftClassPath(TEXT("/Script/Engine.World"));
-		Filter.ObjectClasses.Add(FSoftClassPath(TEXT("/Script/CineCameraRigs.CineSplineMetadata")));
-		SyncConfig->IncludeObjectClassFilters.Add(Filter);
+		const FSoftClassPath MetadataClassPath = FSoftClassPath(TEXT("/Script/CineCameraRigs.CineSplineMetadata"));
+		const bool bIncluded = Algo::AnyOf(SyncConfig->IncludeObjectClassFilters, [MetadataClassPath](const FTransactionClassFilter& Filter)
+			{
+				return Filter.ObjectClasses.Contains(MetadataClassPath);
+			});
+
+		if (!bIncluded)
+		{
+			FTransactionClassFilter Filter;
+			Filter.ObjectOuterClass = FSoftClassPath(TEXT("/Script/Engine.World"));
+			Filter.ObjectClasses.Add(MetadataClassPath);
+			SyncConfig->IncludeObjectClassFilters.Add(Filter);
+		}
+
 	}
 }
 
