CineCameraRigRail:
- Added checkbox option to set all inheritance options at once
- Visualize and manipulate point rotation
- Drive mode implemented in the rig rail itself
- Speed Visualization in the rig rail itself

diff --git a/Engine/Plugins/Experimental/CineCameraRigs/Config/BaseCineCameraRigs.ini b/Engine/Plugins/Experimental/CineCameraRigs/Config/BaseCineCameraRigs.ini
--- a/Engine/Plugins/Experimental/CineCameraRigs/Config/BaseCineCameraRigs.ini
+++ b/Engine/Plugins/Experimental/CineCameraRigs/Config/BaseCineCameraRigs.ini
@@ -1,3 +1,4 @@
 [/Script/CineCameraRigs.CineCameraRigRailSettings]
 DefaultSplineMeshMaterial=/CineCameraRigs/RigRail/M_RigRail.M_RigRail
-DefaultSplineMeshTexture=/CineCameraRigs/RigRail/T_RigRailDefault.T_RigRailDefault
\ No newline at end of file
+DefaultSplineMeshTexture=/CineCameraRigs/RigRail/T_RigRailDefault.T_RigRailDefault
+SpeedModeSplineMeshTexture=/CineCameraRigs/RigRail/T_RigRailSpeedMode.T_RigRailSpeedMode
\ No newline at end of file
diff --git a/Engine/Plugins/Experimental/CineCameraRigs/Source/CineCameraRigs/Private/CineCameraRigRail.cpp b/Engine/Plugins/Experimental/CineCameraRigs/Source/CineCameraRigs/Private/CineCameraRigRail.cpp
--- a/Engine/Plugins/Experimental/CineCameraRigs/Source/CineCameraRigs/Private/CineCameraRigRail.cpp
+++ b/Engine/Plugins/Experimental/CineCameraRigs/Source/CineCameraRigs/Private/CineCameraRigRail.cpp
@@ -3,13 +3,14 @@
 #include "CineCameraRigRail.h"
 #include "CineCameraRigRailHelpers.h"
 #include "CineCameraRigsSettings.h"
-#include "CineSplineComponent.h"
 
 #include "CineCameraComponent.h"
+#include "CineSplineComponent.h"
 #include "Components/SplineMeshComponent.h"
+
+#include "Engine/World.h"
 #include "Kismet/KismetMathLibrary.h"
 #include "Materials/MaterialInstanceDynamic.h"
-#include "Engine/World.h"
 #include "UObject/Package.h"
 
 ACineCameraRigRail::ACineCameraRigRail(const FObjectInitializer& ObjectInitializer)
@@ -21,6 +22,7 @@ ACineCameraRigRail::ACineCameraRigRail(const FObjectInitializer& ObjectInitializ
 	if (CineSplineComponent != nullptr)
 	{
 		CineSplineComponent->OnSplineEdited.AddUObject(this, &ACineCameraRigRail::OnSplineEdited);
+		CineSplineComponent->Duration = 10.0f;
 	}
 
 	UMaterialInterface* DefaultMaterial = GetDefault<UCineCameraRigRailSettings>()->DefaultSplineMeshMaterial.LoadSynchronous();
@@ -233,6 +235,34 @@ void ACineCameraRigRail::PostEditChangeProperty(FPropertyChangedEvent& PropertyC
 			UpdateSplineMeshMID();
 			SetMIDParameters();
 	}
+	else if (PropertyName == GET_MEMBER_NAME_CHECKED(ACineCameraRigRail, AbsolutePositionOnRail) ||
+			 PropertyName == GET_MEMBER_NAME_CHECKED(ACineCameraRigRail, CurrentPositionOnRail))
+	{
+		SpeedAccumulatedTime = 0.0f;
+		if (Speed != 0.0f)
+		{
+			float TotalTime = CineSplineComponent->GetSplineLength() / Speed;
+			SpeedAccumulatedTime = UKismetMathLibrary::SafeDivide(AbsolutePositionOnRail - StartPositionValue(), LastPositionValue() - StartPositionValue()) * TotalTime;
+		}
+	}
+	else if (PropertyName == GET_MEMBER_NAME_CHECKED(ACineCameraRigRail, Speed))
+	{
+		if (Speed != 0.0)
+		{
+			float TotalTime = CineSplineComponent->GetSplineLength() / Speed;
+			SpeedAccumulatedTime = TotalTime * SpeedProgress;
+		}
+	}
+	else if (PropertyName == GET_MEMBER_NAME_CHECKED(ACineCameraRigRail, bDisplaySpeedHeatmap) ||
+		     PropertyName == GET_MEMBER_NAME_CHECKED(ACineCameraRigRail, SpeedSampleCountPerSegment) || 
+			 PropertyName == GET_MEMBER_NAME_CHECKED(ACineCameraRigRail, DriveMode) || 
+			 PropertyName == GET_MEMBER_NAME_CHECKED(ACineCameraRigRail, bUseAbsolutePosition))
+	{
+		UpdateSpeedHeatmap();
+		UpdateSplineMeshMID();
+		SetMIDParameters();
+	}
+
 }
 #endif
 
@@ -285,6 +315,7 @@ FVector ACineCameraRigRail::GetVelocityAtPosition(const float InPosition, const
 
 void ACineCameraRigRail::OnSplineEdited()
 {
+	UpdateSpeedHeatmap();
 	UpdateSplineMeshMID();
 	SetMIDParameters();
 }
@@ -303,3 +334,161 @@ void ACineCameraRigRail::SetMIDParameters()
 		}
 	}
 }
+
+void ACineCameraRigRail::Tick(float DeltaTime)
+{
+	Super::Tick(DeltaTime);
+	if (DriveMode == ECineCameraRigRailDriveMode::Duration)
+	{
+		DriveByParam(DeltaTime);
+	}
+	else if (DriveMode == ECineCameraRigRailDriveMode::Speed)
+	{
+		DriveBySpeed(DeltaTime);
+	}
+}
+
+void ACineCameraRigRail::DriveByParam(float DeltaTime)
+{
+	const float PositionDuration = LastPositionValue() - StartPositionValue();
+	float Param = PositionDuration * UKismetMathLibrary::SafeDivide(DeltaTime, CineSplineComponent->Duration);
+	Param += bUseAbsolutePosition ? AbsolutePositionOnRail : CurrentPositionOnRail;
+	if (bLoop)
+	{
+		Param = PositionDuration > 0.0f ? FMath::Fmod(Param - StartPositionValue(), PositionDuration) : 0.0;
+		Param += StartPositionValue();
+	}
+	else
+	{
+		Param = FMath::Clamp(Param, StartPositionValue(), LastPositionValue());
+	}
+	bUseAbsolutePosition ? AbsolutePositionOnRail = Param : CurrentPositionOnRail = Param;
+}
+
+void ACineCameraRigRail::DriveBySpeed(float DeltaTime)
+{
+	if (Speed == 0.0f)
+	{
+		return;
+	}
+
+	float TotalTime = CineSplineComponent->GetSplineLength() / Speed;
+
+	SpeedAccumulatedTime += DeltaTime;
+	SpeedAccumulatedTime = bLoop ? FMath::Fmod(SpeedAccumulatedTime, TotalTime) : FMath::Clamp(SpeedAccumulatedTime, 0.0f, TotalTime);
+
+	SpeedProgress = SpeedAccumulatedTime / TotalTime;
+
+	if (bUseAbsolutePosition)
+	{
+		float Distance = CineSplineComponent->GetSplineLength() * SpeedProgress;
+		float InputKey = CineSplineComponent->GetInputKeyValueAtDistanceAlongSpline(Distance);
+		AbsolutePositionOnRail = CineSplineComponent->GetPositionAtInputKey(InputKey);
+	}
+	else
+	{
+		CurrentPositionOnRail = SpeedProgress;
+	}
+}
+
+float ACineCameraRigRail::StartPositionValue() const
+{
+	return bUseAbsolutePosition ? CineSplineComponent->GetFloatPropertyAtSplinePoint(0, FName(TEXT("AbsolutePosition"))) : 0.0f;
+}
+
+float ACineCameraRigRail::LastPositionValue() const
+{
+	int32 const NumPoints = CineSplineComponent->GetNumberOfSplinePoints();
+	return bUseAbsolutePosition ? CineSplineComponent->GetFloatPropertyAtSplinePoint(NumPoints - 1, FName(TEXT("AbsolutePosition"))) : 1.0f;
+}
+
+void ACineCameraRigRail::SetDriveModeSpeed(float Value)
+{
+	float TotalTime = UKismetMathLibrary::SafeDivide(CineSplineComponent->GetSplineLength(), Value);
+	SpeedAccumulatedTime = TotalTime * SpeedProgress;
+	Speed = Value;
+}
+
+void ACineCameraRigRail::SetDriveMode(ECineCameraRigRailDriveMode InMode)
+{
+	DriveMode = InMode;
+	UpdateSpeedHeatmap();
+	UpdateSplineMeshMID();
+	SetMIDParameters();
+}
+
+void ACineCameraRigRail::UpdateSpeedHeatmap()
+{
+	// If bDisplaySpeedHeatmap is false, assigns default grey texture
+	// If DriveMode is Speed, assigns green texture indicating it's constant speed
+	// Otherwise, samples the velocity along the spline and generates a heatmap texture
+
+	if (!bDisplaySpeedHeatmap)
+	{
+		UTexture2D* DefaultTexture = GetDefault<UCineCameraRigRailSettings>()->DefaultSplineMeshTexture.LoadSynchronous();
+		if (DefaultTexture)
+		{
+			SplineMeshTexture = DefaultTexture;
+		}
+		return;
+	}
+	else if (DriveMode == ECineCameraRigRailDriveMode::Speed)
+	{
+		UTexture2D* SpeedModeTexture = GetDefault<UCineCameraRigRailSettings>()->SpeedModeSplineMeshTexture.LoadSynchronous();
+		if (SpeedModeTexture)
+		{
+			SplineMeshTexture = SpeedModeTexture;
+		}
+		return;
+	}
+
+	TArray<float> SpeedValues;
+	const int32 NumPoints = CineSplineComponent->GetNumberOfSplinePoints();
+	const int32 NumSegments = CineSplineComponent->GetNumberOfSplineSegments();
+
+	SpeedValues.SetNum(NumSegments * SpeedSampleCountPerSegment + 1);
+	const float SplineLen = CineSplineComponent->GetSplineLength();
+
+	for (int32 SegmentIndex = 0; SegmentIndex < NumSegments; ++SegmentIndex)
+	{
+		for (int32 PerSegmentIndex = 0; PerSegmentIndex < SpeedSampleCountPerSegment; ++PerSegmentIndex)
+		{
+			int32 Index = SegmentIndex * SpeedSampleCountPerSegment + PerSegmentIndex;
+			const float InputKey = (float)SegmentIndex + (float)PerSegmentIndex / (float)SpeedSampleCountPerSegment;
+			if (bUseAbsolutePosition)
+			{
+				const float Position = CineSplineComponent->GetPositionAtInputKey(InputKey);
+				SpeedValues[Index] = GetVelocityAtPosition(Position).Length();
+			}
+			else
+			{
+				const float Distance = CineSplineComponent->GetDistanceAlongSplineAtSplineInputKey(InputKey);
+				SpeedValues[Index] = GetVelocityAtPosition(UKismetMathLibrary::SafeDivide(Distance, SplineLen)).Length();
+			}
+		}
+	}
+	// Last Point
+	const int32 Index = NumSegments * SpeedSampleCountPerSegment;
+	const float InputKey = (float)NumPoints;
+	if (bUseAbsolutePosition)
+	{
+		const float Position = CineSplineComponent->GetPositionAtInputKey(InputKey);
+		SpeedValues[Index] = GetVelocityAtPosition(Position).Length();
+	}
+	else
+	{
+		const float Distance = CineSplineComponent->GetDistanceAlongSplineAtSplineInputKey(InputKey);
+		SpeedValues[Index] = GetVelocityAtPosition(UKismetMathLibrary::SafeDivide(Distance, SplineLen)).Length();
+	}
+
+	float AverageValue = 0.0f;
+	for (float SpeedValue : SpeedValues)
+	{
+		AverageValue += SpeedValue;
+	}
+	AverageValue /= (float)SpeedValues.Num();
+
+	float LowValue = AverageValue * 0.5f;
+	float HighValue = AverageValue * 2.0f;
+	UCineCameraRigRailHelpers::CreateOrUpdateSplineHeatmapTexture(static_cast<UTexture2D*&>(SplineMeshTexture), SpeedValues, LowValue, AverageValue, HighValue);
+}
\ No newline at end of file
diff --git a/Engine/Plugins/Experimental/CineCameraRigs/Source/CineCameraRigs/Public/CineCameraRigRail.h b/Engine/Plugins/Experimental/CineCameraRigs/Source/CineCameraRigs/Public/CineCameraRigRail.h
--- a/Engine/Plugins/Experimental/CineCameraRigs/Source/CineCameraRigs/Public/CineCameraRigRail.h
+++ b/Engine/Plugins/Experimental/CineCameraRigs/Source/CineCameraRigs/Public/CineCameraRigRail.h
@@ -9,6 +9,19 @@
 
 #include "CineCameraRigRail.generated.h"
 
+UENUM(BlueprintType)
+enum class ECineCameraRigRailDriveMode : uint8
+{
+	/** Manual Mode*/
+	Manual,
+
+	/** Duration Mode. AbsolutePostionOnRail is updated based on the spline duration */
+	Duration,
+
+	/** Speed Mode. AbsolutePositionOnRail is updated based on the specified speed*/
+	Speed,
+};
+
 UCLASS(Blueprintable, Category = "VirtualProduction")
 class CINECAMERARIGS_API ACineCameraRigRail : public ACameraRig_Rail
 {
@@ -17,6 +30,8 @@ class CINECAMERARIGS_API ACineCameraRigRail : public ACameraRig_Rail
 public:
 	ACineCameraRigRail(const FObjectInitializer& ObjectInitializer);
 
+	virtual void Tick(float DeltaTime) override;
+
 	/* Returns CineSplineComponent*/
 	UFUNCTION(BlueprintPure, Category = "Rail Components")
 	UCineSplineComponent* GetCineSplineComponent() const { return CineSplineComponent; }
@@ -34,7 +49,7 @@ public:
 	bool bUsePointRotation = true;
 
 	/* Material assigned to spline component mesh*/
-	UPROPERTY(EditAnywhere, BlueprintSetter=SetSplineMeshMaterial, Category = "SplineVisualization")
+	UPROPERTY(BlueprintSetter=SetSplineMeshMaterial, Category = "SplineVisualization")
 	TObjectPtr<UMaterialInterface> SplineMeshMaterial;
 
 	/* Material Instance Dynamic created for the spline mesh */
@@ -42,9 +57,17 @@ public:
 	TArray<TObjectPtr<UMaterialInstanceDynamic>> SplineMeshMIDs;
 
 	/* Texture that can be set to SplineMeshMIDs */
-	UPROPERTY(EditAnywhere, BlueprintSetter=SetSplineMeshTexture, Category = "SplineVisualization")
+	UPROPERTY(BlueprintSetter=SetSplineMeshTexture, Category = "SplineVisualization")
 	TObjectPtr<UTexture2D> SplineMeshTexture;
 
+	/* Enable speed visualization*/
+	UPROPERTY(EditAnywhere, BlueprintReadWrite, Category = "SplineVisualization")
+	bool bDisplaySpeedHeatmap = true;
+
+	/* Number of speed samples per spline segment*/
+	UPROPERTY(EditAnywhere, BlueprintReadWrite, Category = "SplineVisualization", meta=(ClampMin=1, EditCondition="bDisplaySpeedHeatmap"))
+	int32 SpeedSampleCountPerSegment = 4;
+
 	/* Determines if camera mount inherits LocationX*/
 	UPROPERTY(EditAnywhere, BlueprintReadWrite, Category = "Attachment|Location")
 	bool bAttachLocationX = true;
@@ -81,6 +104,18 @@ public:
 	UPROPERTY(EditAnywhere, BlueprintReadWrite, Category = "Attachment|Camera")
 	bool bInheritFocusDistance = true;
 
+	/* Drive Mode to update position in tick*/
+	UPROPERTY(EditAnywhere, BlueprintSetter=SetDriveMode, Category = "DriveMode")
+	ECineCameraRigRailDriveMode DriveMode = ECineCameraRigRailDriveMode::Manual;
+
+	/* Specifies the drive speed of the rig rail in centimeter per second */
+	UPROPERTY(EditAnywhere, BlueprintSetter=SetDriveModeSpeed, Category = "DriveMode")
+	float Speed = 100;
+
+	/* Enable loop in speed or duration mode */
+	UPROPERTY(EditAnywhere, BlueprintReadWrite, Category = "DriveMode")
+	bool bLoop = true;
+
 	/* Set spline mesh material*/
 	UFUNCTION(BlueprintSetter)
 	void SetSplineMeshMaterial(UMaterialInterface* InMaterial);
@@ -93,17 +128,34 @@ public:
 	UFUNCTION(BlueprintCallable, Category = "CineCameraRigRail")
 	FVector GetVelocityAtPosition(const float InPosition, const float delta = 0.001) const;
 
+	/* Set drive mode speed value*/
+	UFUNCTION(BlueprintSetter)
+	void SetDriveModeSpeed(float Value);
+
+	/* Set drive mode*/
+	UFUNCTION(BlueprintSetter)
+	void SetDriveMode(ECineCameraRigRailDriveMode InMode);
+
 #if WITH_EDITOR
 	virtual void PostEditChangeProperty(FPropertyChangedEvent& PropertyChangedEvent) override;
 #endif
 
 protected:
-	class UCineSplineComponent* CineSplineComponent;
+	UPROPERTY(Category = "CineSpline", VisibleAnywhere, BlueprintReadOnly)
+	TObjectPtr<class UCineSplineComponent> CineSplineComponent;
 
 	virtual void UpdateRailComponents() override;
 
 	void UpdateSplineMeshMID();
 	void SetMIDParameters();
+	void UpdateSpeedHeatmap();
 
 	void OnSplineEdited();
+
+	void DriveByParam(float DeltaTime);
+	void DriveBySpeed(float DeltaTime);
+	float StartPositionValue() const;
+	float LastPositionValue() const;
+	float SpeedProgress = 0.0f;
+	float SpeedAccumulatedTime = 0.0f;
 };
diff --git a/Engine/Plugins/Experimental/CineCameraRigs/Source/CineCameraRigs/Public/CineCameraRigsSettings.h b/Engine/Plugins/Experimental/CineCameraRigs/Source/CineCameraRigs/Public/CineCameraRigsSettings.h
--- a/Engine/Plugins/Experimental/CineCameraRigs/Source/CineCameraRigs/Public/CineCameraRigsSettings.h
+++ b/Engine/Plugins/Experimental/CineCameraRigs/Source/CineCameraRigs/Public/CineCameraRigsSettings.h
@@ -15,9 +15,15 @@ class CINECAMERARIGS_API UCineCameraRigRailSettings : public UObject
 {
 	GENERATED_BODY()
 public:
+	/* Path to the default spline mesh material used in CineCameraRigRail*/
 	UPROPERTY(config, EditAnywhere, BlueprintReadWrite, Category = CineCameraRigRail)
 	TSoftObjectPtr<UMaterialInterface> DefaultSplineMeshMaterial;
 
+	/* Path to the default spline mesh texture used when speed visualization is off*/
 	UPROPERTY(config, EditAnywhere, BlueprintReadWrite, Category = CineCameraRigRail)
 	TSoftObjectPtr<UTexture2D> DefaultSplineMeshTexture;
+
+	/* Path to the texture used in the speed visualization when drive mode is Speed*/
+	UPROPERTY(config, EditAnywhere, BlueprintReadWrite, Category = CineCameraRigRail)
+	TSoftObjectPtr<UTexture2D> SpeedModeSplineMeshTexture;
 };
\ No newline at end of file
diff --git a/Engine/Plugins/Experimental/CineCameraRigs/Source/CineCameraRigs/Public/CineSplineComponent.h b/Engine/Plugins/Experimental/CineCameraRigs/Source/CineCameraRigs/Public/CineSplineComponent.h
--- a/Engine/Plugins/Experimental/CineCameraRigs/Source/CineCameraRigs/Public/CineSplineComponent.h
+++ b/Engine/Plugins/Experimental/CineCameraRigs/Source/CineCameraRigs/Public/CineSplineComponent.h
@@ -113,6 +113,21 @@ public:
 	/** Update the spline tangents and SplineReparamTable. And trigers OnSplineEdited delegates */
 	virtual void UpdateSpline() override;
 
+#if WITH_EDITORONLY_DATA
+	/* Toggle if visualizer displays normalized position or absolute position */
+	UPROPERTY(EditAnywhere, Category = "CineSplineEditor", meta = (InlineEditConditionToggle = true))
+	bool bShouldVisualizeNormalizedPosition;
+
+	/* Whether spline length value should be displayed in visualizer */
+	UPROPERTY(EditAnywhere, Category = "CineSplineEditor", meta = (InlineEditConditionToggle = true))
+	bool bShouldVisualizeSplineLength;
+
+	/* Whether point rotation arrows should be displayed in visualizer */
+	UPROPERTY(EditAnywhere, Category = "CineSplineEditor", meta = (InlineEditConditionToggle = true))
+	bool bShouldVisualizePointRotation;
+
+#endif
+
 protected:
 	void SynchronizeProperties();
 
diff --git a/Engine/Plugins/Experimental/CineCameraRigs/Source/CineCameraRigsEditor/Private/CineCameraRigRailDetails.cpp b/Engine/Plugins/Experimental/CineCameraRigs/Source/CineCameraRigsEditor/Private/CineCameraRigRailDetails.cpp
--- a/Engine/Plugins/Experimental/CineCameraRigs/Source/CineCameraRigsEditor/Private/CineCameraRigRailDetails.cpp
+++ b/Engine/Plugins/Experimental/CineCameraRigs/Source/CineCameraRigsEditor/Private/CineCameraRigRailDetails.cpp
@@ -7,12 +7,23 @@
 #include "DetailLayoutBuilder.h"
 #include "DetailWidgetRow.h"
 #include "Editor.h"
+#include "IDetailGroup.h"
+#include "Internationalization/Internationalization.h"
+#include "Layout/Visibility.h"
 #include "PropertyHandle.h"
+#include "ScopedTransaction.h"
 #include "UObject/PropertyIterator.h"
+#include "Widgets/Input/SCheckBox.h"
 #include "Widgets/Input/SNumericEntryBox.h"
-#include "Internationalization/Internationalization.h"
-#include "ScopedTransaction.h"
-#include "Layout/Visibility.h"
+
+#include "ComponentVisualizer.h"
+#include "CineSplineComponentVisualizer.h"
+
+#include "Editor/UnrealEdEngine.h"
+#include "UnrealEdGlobals.h"
+
+#include "Algo/AllOf.h"
+#include "Algo/NoneOf.h"
 
 #define LOCTEXT_NAMESPACE "FCineCameraRigRailDetails"
 
@@ -33,29 +44,52 @@ void FCineCameraRigRailDetails::CustomizeDetails(IDetailLayoutBuilder& DetailBui
 	}
 	RigRailActorPtr = RigRailActor;
 
+	CustomizeRailControlCategory(DetailBuilder);
+	CustomizeSplineVisualizationCategory(DetailBuilder);
+	CustomizeAttachmentCategory(DetailBuilder);
+	CustomizeDriveModeCategory(DetailBuilder);
+	HideExtraCategories(DetailBuilder);
+}
+
+void FCineCameraRigRailDetails::CustomizeRailControlCategory(IDetailLayoutBuilder& DetailBuilder)
+{
 	TSharedRef<IPropertyHandle> AbsolutePositionPropertyHandle = DetailBuilder.GetProperty(GET_MEMBER_NAME_CHECKED(ACineCameraRigRail, AbsolutePositionOnRail));
 	TSharedRef<IPropertyHandle> UseAbsolutePositionPropertyHandle = DetailBuilder.GetProperty(GET_MEMBER_NAME_CHECKED(ACineCameraRigRail, bUseAbsolutePosition));
 	TSharedRef<IPropertyHandle> OrigPositionPropertyHandle = DetailBuilder.GetProperty(GET_MEMBER_NAME_CHECKED(ACameraRig_Rail, CurrentPositionOnRail), ACameraRig_Rail::StaticClass());
 	TSharedRef<IPropertyHandle> LockOrientationPropertyHandle = DetailBuilder.GetProperty(GET_MEMBER_NAME_CHECKED(ACameraRig_Rail, bLockOrientationToRail), ACameraRig_Rail::StaticClass());
-	TSharedRef<IPropertyHandle> ShowRailPropertyHandle = DetailBuilder.GetProperty(GET_MEMBER_NAME_CHECKED(ACameraRig_Rail, bShowRailVisualization), ACameraRig_Rail::StaticClass());
-	TSharedRef<IPropertyHandle> PreviewScalePropertyHandle = DetailBuilder.GetProperty(GET_MEMBER_NAME_CHECKED(ACameraRig_Rail, PreviewMeshScale), ACameraRig_Rail::StaticClass());
-
-
 
 	const TAttribute<bool> EditCondition = TAttribute<bool>::Create([this, UseAbsolutePositionPropertyHandle]()
-	{
-		bool bCond = false;
-		UseAbsolutePositionPropertyHandle->GetValue(bCond);
-		return !bCond;
-	});
+		{
+			bool bCond = false;
+			UseAbsolutePositionPropertyHandle->GetValue(bCond);
+			return !bCond;
+		});
 
 	DetailBuilder.EditDefaultProperty(LockOrientationPropertyHandle)->Visibility(EVisibility::Hidden);
 
 	IDetailCategoryBuilder& RailControlsCategory = DetailBuilder.EditCategory("Rail Controls");
-	RailControlsCategory.AddProperty(UseAbsolutePositionPropertyHandle);
+	IDetailPropertyRow& UseAbsolutePositionRow = RailControlsCategory.AddProperty(UseAbsolutePositionPropertyHandle);
 	IDetailPropertyRow& AbsolutePositionRow = RailControlsCategory.AddProperty(AbsolutePositionPropertyHandle);
 	RailControlsCategory.AddProperty(OrigPositionPropertyHandle)
-	.EditCondition(EditCondition, nullptr);
+		.EditCondition(EditCondition, nullptr);
+
+	UseAbsolutePositionRow.CustomWidget(false)
+		.NameContent()
+		[
+			UseAbsolutePositionPropertyHandle->CreatePropertyNameWidget()
+		]
+		.ValueContent()
+		.MinDesiredWidth(125.0f)
+		[
+			SNew(SCheckBox)
+			.IsChecked_Lambda([UseAbsolutePositionPropertyHandle]() -> ECheckBoxState
+				{
+					bool IsChecked;
+					UseAbsolutePositionPropertyHandle.Get().GetValue(IsChecked);
+					return IsChecked ? ECheckBoxState::Checked : ECheckBoxState::Unchecked;
+				})
+			.OnCheckStateChanged(this, &FCineCameraRigRailDetails::OnUseAbsolutePositionChanged, UseAbsolutePositionPropertyHandle)
+		];
 
 	AbsolutePositionRow.CustomWidget(false)
 		.NameContent()
@@ -77,12 +111,194 @@ void FCineCameraRigRailDetails::CustomizeDetails(IDetailLayoutBuilder& DetailBui
 			.OnEndSliderMovement(this, &FCineCameraRigRailDetails::OnEndAbsolutePositionSliderMovement)
 			.Font(IDetailLayoutBuilder::GetDetailFont())
 		];
+}
 
+void FCineCameraRigRailDetails::CustomizeSplineVisualizationCategory(IDetailLayoutBuilder& DetailBuilder)
+{
 	IDetailCategoryBuilder& VisualizationCategory = DetailBuilder.EditCategory("SplineVisualization");
+	TSharedRef<IPropertyHandle> ShowRailPropertyHandle = DetailBuilder.GetProperty(GET_MEMBER_NAME_CHECKED(ACameraRig_Rail, bShowRailVisualization), ACameraRig_Rail::StaticClass());
+	TSharedRef<IPropertyHandle> PreviewScalePropertyHandle = DetailBuilder.GetProperty(GET_MEMBER_NAME_CHECKED(ACameraRig_Rail, PreviewMeshScale), ACameraRig_Rail::StaticClass());
+	TSharedRef<IPropertyHandle> DisplaySpeedHeatmapPropertyHandle = DetailBuilder.GetProperty(GET_MEMBER_NAME_CHECKED(ACineCameraRigRail, bDisplaySpeedHeatmap));
+	TSharedRef<IPropertyHandle> SampleCountPropertyHandle = DetailBuilder.GetProperty(GET_MEMBER_NAME_CHECKED(ACineCameraRigRail, SpeedSampleCountPerSegment));
 	VisualizationCategory.AddProperty(ShowRailPropertyHandle);
+	VisualizationCategory.AddProperty(DisplaySpeedHeatmapPropertyHandle);
+	VisualizationCategory.AddProperty(SampleCountPropertyHandle);
 	VisualizationCategory.AddProperty(PreviewScalePropertyHandle);
 }
 
+void FCineCameraRigRailDetails::CustomizeAttachmentCategory(IDetailLayoutBuilder& DetailBuilder)
+{
+	IDetailCategoryBuilder& AttachmentCategory = DetailBuilder.EditCategory("Attachment");
+
+	// Attachment : Location
+	TArray<TSharedRef<IPropertyHandle>> LocationPropertyHandles;
+	LocationPropertyHandles.Add(DetailBuilder.GetProperty(GET_MEMBER_NAME_CHECKED(ACineCameraRigRail, bAttachLocationX)));
+	LocationPropertyHandles.Add(DetailBuilder.GetProperty(GET_MEMBER_NAME_CHECKED(ACineCameraRigRail, bAttachLocationY)));
+	LocationPropertyHandles.Add(DetailBuilder.GetProperty(GET_MEMBER_NAME_CHECKED(ACineCameraRigRail, bAttachLocationZ)));
+
+	IDetailGroup& LocationGroup = AttachmentCategory.AddGroup(TEXT("Location"), LOCTEXT("LocationLabel", "Location"));
+	LocationGroup.HeaderRow()
+		.NameContent()
+			[
+			SNew(STextBlock)
+			.Text(LOCTEXT("LocationLabel", "Location"))
+			.Font(IDetailLayoutBuilder::GetDetailFont())
+			]
+		.ValueContent()
+			[
+			SNew(SCheckBox)
+			.IsChecked(this, &FCineCameraRigRailDetails::IsAttachOptionChecked, LocationPropertyHandles)
+			.OnCheckStateChanged(this, &FCineCameraRigRailDetails::OnAttachOptionChanged, LocationPropertyHandles)
+			];
+
+	for (auto& PropertyHandle : LocationPropertyHandles)
+	{
+		DetailBuilder.HideProperty(PropertyHandle);
+		LocationGroup.AddPropertyRow(PropertyHandle);
+	}
+
+
+	// Attachment : Rotation
+	TArray<TSharedRef<IPropertyHandle>> RotationPropertyHandles;
+	RotationPropertyHandles.Add(DetailBuilder.GetProperty(GET_MEMBER_NAME_CHECKED(ACineCameraRigRail, bAttachRotationX)));
+	RotationPropertyHandles.Add(DetailBuilder.GetProperty(GET_MEMBER_NAME_CHECKED(ACineCameraRigRail, bAttachRotationY)));
+	RotationPropertyHandles.Add(DetailBuilder.GetProperty(GET_MEMBER_NAME_CHECKED(ACineCameraRigRail, bAttachRotationZ)));
+
+	IDetailGroup& RotationGroup = AttachmentCategory.AddGroup(TEXT("Rotation"), LOCTEXT("RotationLabel", "Rotation"));
+	RotationGroup.HeaderRow()
+		.NameContent()
+		[
+			SNew(STextBlock)
+			.Text(LOCTEXT("RotationLabel", "Rotation"))
+			.Font(IDetailLayoutBuilder::GetDetailFont())
+		]
+		.ValueContent()
+		[
+			SNew(SCheckBox)
+			.IsChecked(this, &FCineCameraRigRailDetails::IsAttachOptionChecked, RotationPropertyHandles)
+			.OnCheckStateChanged(this, &FCineCameraRigRailDetails::OnAttachOptionChanged, RotationPropertyHandles)
+		];
+
+	for (auto& PropertyHandle : RotationPropertyHandles)
+	{
+		DetailBuilder.HideProperty(PropertyHandle);
+		RotationGroup.AddPropertyRow(PropertyHandle);
+	}
+
+	// Attachment : Camera
+	TArray<TSharedRef<IPropertyHandle>> CameraPropertyHandles;
+	CameraPropertyHandles.Add(DetailBuilder.GetProperty(GET_MEMBER_NAME_CHECKED(ACineCameraRigRail, bInheritFocalLength)));
+	CameraPropertyHandles.Add(DetailBuilder.GetProperty(GET_MEMBER_NAME_CHECKED(ACineCameraRigRail, bInheritAperture)));
+	CameraPropertyHandles.Add(DetailBuilder.GetProperty(GET_MEMBER_NAME_CHECKED(ACineCameraRigRail, bInheritFocusDistance)));
+
+	IDetailGroup& CameraGroup = AttachmentCategory.AddGroup(TEXT("Camera"), LOCTEXT("CameraLabel", "Camera"));
+	CameraGroup.HeaderRow()
+		.NameContent()
+		[
+			SNew(STextBlock)
+			.Text(LOCTEXT("CameraLabel", "Camera"))
+			.Font(IDetailLayoutBuilder::GetDetailFont())
+		]
+		.ValueContent()
+		[
+			SNew(SCheckBox)
+			.IsChecked(this, &FCineCameraRigRailDetails::IsAttachOptionChecked, CameraPropertyHandles)
+			.OnCheckStateChanged(this, &FCineCameraRigRailDetails::OnAttachOptionChanged, CameraPropertyHandles)
+		];
+
+	for (auto& PropertyHandle : CameraPropertyHandles)
+	{
+		DetailBuilder.HideProperty(PropertyHandle);
+		CameraGroup.AddPropertyRow(PropertyHandle);
+	}
+}
+
+void FCineCameraRigRailDetails::CustomizeDriveModeCategory(IDetailLayoutBuilder& DetailBuilder)
+{
+	IDetailCategoryBuilder& DriveModeCategory = DetailBuilder.EditCategory("DriveMode");
+	DriveModeHandle = DetailBuilder.GetProperty(GET_MEMBER_NAME_CHECKED(ACineCameraRigRail, DriveMode));
+	TSharedPtr<IPropertyHandle> DurationHandle = DetailBuilder.GetProperty(GET_MEMBER_NAME_CHECKED(USplineComponent, Duration), USplineComponent::StaticClass());
+	TSharedPtr<IPropertyHandle> SpeedHandle = DetailBuilder.GetProperty(GET_MEMBER_NAME_CHECKED(ACineCameraRigRail, Speed));
+	TSharedPtr<IPropertyHandle> LoopHandle = DetailBuilder.GetProperty(GET_MEMBER_NAME_CHECKED(ACineCameraRigRail, bLoop));
+
+	DriveModeCategory.AddProperty(DriveModeHandle.ToSharedRef());
+
+	DriveModeCategory.AddProperty(DurationHandle.ToSharedRef())
+		.EditCondition(TAttribute<bool>::Create([this]() { return GetDriveMode() == ECineCameraRigRailDriveMode::Duration; }), nullptr);
+	DriveModeCategory.AddProperty(SpeedHandle.ToSharedRef())
+		.EditCondition(TAttribute<bool>::Create([this]() { return GetDriveMode() == ECineCameraRigRailDriveMode::Speed; }), nullptr);
+	DriveModeCategory.AddProperty(LoopHandle.ToSharedRef())
+		.EditCondition(TAttribute<bool>::Create([this]() { return GetDriveMode() != ECineCameraRigRailDriveMode::Manual; }), nullptr);
+
+}
+
+void FCineCameraRigRailDetails::HideExtraCategories(IDetailLayoutBuilder& DetailBuilder)
+{
+	DetailBuilder.HideCategory("Editor");
+	DetailBuilder.HideCategory("Selected Points");
+	DetailBuilder.HideCategory("HLOD");
+	DetailBuilder.HideCategory("PathTracing");
+	DetailBuilder.HideCategory("Spline");
+	DetailBuilder.HideCategory("Navigation");
+	DetailBuilder.HideCategory("Tags");
+	DetailBuilder.HideCategory("Cooking");
+	DetailBuilder.HideCategory("LOD");
+	DetailBuilder.HideCategory("TextureStreaming");
+	DetailBuilder.HideCategory("RayTracing");
+	DetailBuilder.HideCategory("AssetUserData");
+}
+
+ECineCameraRigRailDriveMode FCineCameraRigRailDetails::GetDriveMode() const
+{
+	if (!DriveModeHandle.IsValid())
+	{
+		return ECineCameraRigRailDriveMode::Manual;
+	}
+	uint8 DriveModeValue;
+	DriveModeHandle->GetValue(DriveModeValue);
+	return static_cast<ECineCameraRigRailDriveMode>(DriveModeValue);
+}
+
+ECheckBoxState FCineCameraRigRailDetails::IsAttachOptionChecked(TArray<TSharedRef<IPropertyHandle>> PropertyHandles) const
+{
+	const bool bAllChecked = Algo::AllOf(PropertyHandles, [](const TSharedRef<IPropertyHandle>& PropertyHandle)
+		{
+			bool bValue;
+			PropertyHandle->GetValue(bValue);
+			return bValue;
+		});
+	const bool bNoneChecked = Algo::NoneOf(PropertyHandles, [](const TSharedRef<IPropertyHandle>& PropertyHandle)
+		{
+			bool bValue;
+			PropertyHandle->GetValue(bValue);
+			return bValue;
+		});
+
+	if (bAllChecked)
+	{
+		return ECheckBoxState::Checked;
+	}
+	else if (bNoneChecked)
+	{
+		return ECheckBoxState::Unchecked;
+	}
+	return ECheckBoxState::Undetermined;
+}
+
+void FCineCameraRigRailDetails::OnAttachOptionChanged(ECheckBoxState NewState, TArray<TSharedRef<IPropertyHandle>> PropertyHandles)
+{
+	if (NewState == ECheckBoxState::Undetermined)
+	{
+		return;
+	}
+
+	const bool bValue = (NewState == ECheckBoxState::Checked) ? true : false;
+	for (TSharedRef<IPropertyHandle>& PropertyHandle : PropertyHandles)
+	{
+		PropertyHandle->SetValue(bValue);
+	}
+}
+
 void FCineCameraRigRailDetails::OnAbsolutePositionChanged(float NewValue)
 {
 	if (ACineCameraRigRail* RigRailActor = RigRailActorPtr.Get())
@@ -98,6 +314,27 @@ void FCineCameraRigRailDetails::OnAbsolutePositionChanged(float NewValue)
 	}
 }
 
+void FCineCameraRigRailDetails::OnUseAbsolutePositionChanged(ECheckBoxState NewState, TSharedRef<IPropertyHandle> PropertyHandle)
+{
+	if (NewState == ECheckBoxState::Undetermined)
+	{
+		return;
+	}
+	const bool bValue = (NewState == ECheckBoxState::Checked) ? true : false;
+	if (PropertyHandle->IsValidHandle())
+	{
+		PropertyHandle->SetValue(bValue);
+	}
+	
+	if (ACineCameraRigRail* RigRailActor = RigRailActorPtr.Get())
+	{
+		if (UCineSplineComponent* SplineComp = RigRailActor->GetCineSplineComponent())
+		{
+			SplineComp->bShouldVisualizeNormalizedPosition = !bValue;
+		}
+	}
+}
+
 void FCineCameraRigRailDetails::OnAbsolutePositionCommitted(float NewValue, ETextCommit::Type CommitType)
 {
 	if (ACineCameraRigRail* RigRailActor = RigRailActorPtr.Get())
diff --git a/Engine/Plugins/Experimental/CineCameraRigs/Source/CineCameraRigsEditor/Private/CineCameraRigRailDetails.h b/Engine/Plugins/Experimental/CineCameraRigs/Source/CineCameraRigsEditor/Private/CineCameraRigRailDetails.h
--- a/Engine/Plugins/Experimental/CineCameraRigs/Source/CineCameraRigsEditor/Private/CineCameraRigRailDetails.h
+++ b/Engine/Plugins/Experimental/CineCameraRigs/Source/CineCameraRigsEditor/Private/CineCameraRigRailDetails.h
@@ -4,8 +4,12 @@
 
 #include "CoreMinimal.h"
 #include "IDetailCustomization.h"
+#include "Styling/SlateTypes.h"
+
 #include "CineCameraRigRail.h"
 
+class IPropertyHandle;
+
 namespace ETextCommit { enum Type : int; }
 
 /**
@@ -37,4 +41,17 @@ private:
 
 	bool bAbsolutePositionSliderStartedTransaction = false;
 
+	ECheckBoxState IsAttachOptionChecked(TArray<TSharedRef<IPropertyHandle>> PropertyHandles) const;
+	void OnAttachOptionChanged(ECheckBoxState NewState, TArray<TSharedRef<IPropertyHandle>> PropertyHandles);
+	
+	TSharedPtr<IPropertyHandle> DriveModeHandle;
+	ECineCameraRigRailDriveMode GetDriveMode() const;
+
+	void OnUseAbsolutePositionChanged(ECheckBoxState NewState, TSharedRef<IPropertyHandle> PropertyHandle);
+
+	void CustomizeRailControlCategory(IDetailLayoutBuilder& DetailBuilder);
+	void CustomizeSplineVisualizationCategory(IDetailLayoutBuilder& DetailBuilder);
+	void CustomizeAttachmentCategory(IDetailLayoutBuilder& DetailBuilder);
+	void CustomizeDriveModeCategory(IDetailLayoutBuilder& DetailBuilder);
+	void HideExtraCategories(IDetailLayoutBuilder& DetailBuilder);
 };
diff --git a/Engine/Plugins/Experimental/CineCameraRigs/Source/CineCameraRigsEditor/Private/CineSplineComponentVisualizer.cpp b/Engine/Plugins/Experimental/CineCameraRigs/Source/CineCameraRigsEditor/Private/CineSplineComponentVisualizer.cpp
--- a/Engine/Plugins/Experimental/CineCameraRigs/Source/CineCameraRigsEditor/Private/CineSplineComponentVisualizer.cpp
+++ b/Engine/Plugins/Experimental/CineCameraRigs/Source/CineCameraRigsEditor/Private/CineSplineComponentVisualizer.cpp
@@ -5,14 +5,88 @@
 #include "Math/Color.h"
 #include "CanvasTypes.h"
 #include "CanvasItem.h"
-#include "SceneView.h"
+#include "EditorViewportClient.h"
 #include "Engine/Engine.h"
+#include "Framework/MultiBox/MultiBoxBuilder.h"
+#include "Kismet/KismetMathLibrary.h"
+#include "SceneManagement.h"
+#include "SceneView.h"
 
 #define LOCTEXT_NAMESPACE "SplineComponentVisualizer"
 
+class FCineSplineComponentVisualizerCommands : public TCommands< FCineSplineComponentVisualizerCommands>
+{
+public:
+	FCineSplineComponentVisualizerCommands() : TCommands< FCineSplineComponentVisualizerCommands>
+		(
+			"CineSplineComponentVisualizer",
+			LOCTEXT("CineSplineComponentVisualizer", "CineSpline Component Visualizer"),
+			NAME_None,
+			FAppStyle::GetAppStyleSetName()
+			)
+	{}
+
+	virtual void RegisterCommands() override
+	{
+		UI_COMMAND(VisualizeNormalizedPosition, "Show Normalized Position Value", "Whether the visualization should show normalized position or absolute position.", EUserInterfaceActionType::ToggleButton, FInputChord());
+		UI_COMMAND(VisualizeSplineLength, "Show Spline Length", "Whether the visualization should show spline length.", EUserInterfaceActionType::ToggleButton, FInputChord());
+		UI_COMMAND(VisualizePointRotation, "Show Point Rotation", "Whether the visualization should show point rotation arrows.", EUserInterfaceActionType::ToggleButton, FInputChord());
+		UI_COMMAND(ManipulatePointRotation, "Manipulate Point Rotation", "Use rotation manipulator to edit PointRotation metadata.", EUserInterfaceActionType::ToggleButton, FInputChord());
+	}
+
+public:
+
+	TSharedPtr<FUICommandInfo> VisualizeNormalizedPosition;
+	TSharedPtr<FUICommandInfo> VisualizeSplineLength;
+	TSharedPtr<FUICommandInfo> VisualizePointRotation;
+	TSharedPtr<FUICommandInfo> ManipulatePointRotation;
+};
+
+FCineSplineComponentVisualizer::FCineSplineComponentVisualizer()
+	: FSplineComponentVisualizer()
+{
+	FCineSplineComponentVisualizerCommands::Register();
+	CineSplineComponentVisualizerActions = MakeShareable(new FUICommandList);
+}
+
 void FCineSplineComponentVisualizer::OnRegister()
 {
 	FSplineComponentVisualizer::OnRegister();
+
+	const auto& Commands = FCineSplineComponentVisualizerCommands::Get();
+
+	SplineComponentVisualizerActions->MapAction(
+		Commands.VisualizeNormalizedPosition,
+		FExecuteAction::CreateSP(this, &FCineSplineComponentVisualizer::OnSetVisualizeNormalizedPosition),
+		FCanExecuteAction(),
+		FIsActionChecked::CreateSP(this, &FCineSplineComponentVisualizer::IsVisualizeNormalizedPosition)
+	);
+
+	SplineComponentVisualizerActions->MapAction(
+		Commands.VisualizeSplineLength,
+		FExecuteAction::CreateSP(this, &FCineSplineComponentVisualizer::OnSetVisualizeSplineLength),
+		FCanExecuteAction(),
+		FIsActionChecked::CreateSP(this, &FCineSplineComponentVisualizer::IsVisualizeSplineLength)
+	);
+
+	SplineComponentVisualizerActions->MapAction(
+		Commands.VisualizePointRotation,
+		FExecuteAction::CreateSP(this, &FCineSplineComponentVisualizer::OnSetVisualizePointRotation),
+		FCanExecuteAction(),
+		FIsActionChecked::CreateSP(this, &FCineSplineComponentVisualizer::IsVisualizePointRotation)
+	);
+
+	SplineComponentVisualizerActions->MapAction(
+		Commands.ManipulatePointRotation,
+		FExecuteAction::CreateSP(this, &FCineSplineComponentVisualizer::OnSetManipulatePointRotation),
+		FCanExecuteAction(),
+		FIsActionChecked::CreateSP(this, &FCineSplineComponentVisualizer::IsManipulatePointRotation)
+	);
+}
+
+FCineSplineComponentVisualizer::~FCineSplineComponentVisualizer()
+{
+	FCineSplineComponentVisualizerCommands::Unregister();
 }
 
 void FCineSplineComponentVisualizer::DrawVisualizationHUD(const UActorComponent* Component, const FViewport* Viewport, const FSceneView* View, FCanvas* Canvas)
@@ -26,7 +100,7 @@ void FCineSplineComponentVisualizer::DrawVisualizationHUD(const UActorComponent*
 	const FIntRect CanvasRect = Canvas->GetViewRect();
 	const float HalfX = CanvasRect.Width() / 2.f;
 	const float HalfY = CanvasRect.Height() / 2.f;
-
+	const float TextPositionOffset = 12.0f;
 
 	if (const UCineSplineComponent* SplineComp = Cast<const UCineSplineComponent>(Component))
 	{
@@ -45,12 +119,187 @@ void FCineSplineComponentVisualizer::DrawVisualizationHUD(const UActorComponent*
 				FNumberFormattingOptions FmtOptions;
 				FmtOptions.SetMaximumFractionalDigits(4);
 				FmtOptions.SetMinimumFractionalDigits(1);
-				const FText Text = FText::AsNumber(Metadata->AbsolutePosition.Points[i].OutVal, &FmtOptions);
-				FLinearColor Color = FMath::IsNearlyEqual(FMath::Frac(Metadata->AbsolutePosition.Points[i].OutVal), 0.0f) ? FLinearColor(1,1,0,1) : FLinearColor(0, 1, 1, 1);
+
+				// Draw normalized position values or AbsolutePosition metadata values, based on bUseAbsolutePosition parameter in CineCameraRigRailActor
+				const bool bShowNormalizedPosition = SplineComp->bShouldVisualizeNormalizedPosition;
+				const float Value = bShowNormalizedPosition ? UKismetMathLibrary::SafeDivide(SplineComp->GetDistanceAlongSplineAtSplinePoint(i), SplineComp->GetSplineLength()) : Metadata->AbsolutePosition.Points[i].OutVal;
+				const FLinearColor Color = bShowNormalizedPosition ? FLinearColor(0, 1, 1, 1) : FMath::IsNearlyEqual(FMath::Frac(Metadata->AbsolutePosition.Points[i].OutVal), 0.0f) ? FLinearColor(1, 1, 0, 1) : FLinearColor(0, 1, 1, 1);
+				const FText Text = FText::AsNumber(Value , &FmtOptions);
 				Canvas->DrawShadowedString(DrawPositionX, DrawPositionY, *Text.ToString(), GEngine->GetLargeFont(), Color, FLinearColor::Black);
+
+				if (SplineComp->bShouldVisualizeSplineLength)
+				{
+					const float Distance = SplineComp->GetDistanceAlongSplineAtSplinePoint(i);
+					const FText DistanceText = FText::AsNumber(Distance, &FmtOptions);
+					Canvas->DrawShadowedString(DrawPositionX, (DrawPositionY+TextPositionOffset), *DistanceText.ToString(), GEngine->GetLargeFont(), FLinearColor::White, FLinearColor::Black);
+				}
 			}
 		}
 	}
 }
 
+
+void FCineSplineComponentVisualizer::DrawVisualization(
+	const UActorComponent* Component,
+	const FSceneView* View,
+	FPrimitiveDrawInterface* PDI
+)
+{
+	FSplineComponentVisualizer::DrawVisualization(Component, View, PDI);
+
+	const UCineSplineComponent* CineSplineComponent = Cast<const UCineSplineComponent>(Component);
+
+	if (CineSplineComponent && CineSplineComponent->bShouldVisualizePointRotation)
+	{
+		const int32 NumSplinePoints = CineSplineComponent->GetNumberOfSplinePoints();
+		const float ArrowLength = 100.0f;
+		const float ArrowSize = 5.0f;
+		const FColor ArrowColor = FColor::Cyan;
+
+		for (int32 PointIdx = 0; PointIdx < NumSplinePoints; PointIdx++)
+		{
+			const FVector Location = CineSplineComponent->GetLocationAtSplinePoint(PointIdx, ESplineCoordinateSpace::World);
+			const FQuat Rotation = CineSplineComponent->GetPointRotationAtSplinePoint(PointIdx);
+
+			const FVector ArrowDirection = Rotation.RotateVector(FVector::ForwardVector);
+			const FVector ArrowEndLocation = Location + ArrowDirection * ArrowLength;
+			const FTransform Transform = FTransform(Rotation, Location);
+
+			PDI->DrawLine(Location, ArrowEndLocation, ArrowColor, SDPG_World, 1.0f, 0.0f, true);
+			DrawDirectionalArrow(PDI, Transform.ToMatrixNoScale(), ArrowColor, ArrowLength, ArrowSize, SDPG_World);
+		}
+	}
+}
+
+void FCineSplineComponentVisualizer::GenerateContextMenuSections(FMenuBuilder& InMenuBuilder) const
+{
+	FSplineComponentVisualizer::GenerateContextMenuSections(InMenuBuilder);
+
+	InMenuBuilder.BeginSection("CineSpline", LOCTEXT("CineSpline", "CineSpline"));
+	{
+		InMenuBuilder.AddMenuEntry(FCineSplineComponentVisualizerCommands::Get().VisualizeNormalizedPosition);
+		InMenuBuilder.AddMenuEntry(FCineSplineComponentVisualizerCommands::Get().VisualizeSplineLength);
+		InMenuBuilder.AddMenuEntry(FCineSplineComponentVisualizerCommands::Get().VisualizePointRotation);
+		InMenuBuilder.AddMenuEntry(FCineSplineComponentVisualizerCommands::Get().ManipulatePointRotation);
+	}
+	InMenuBuilder.EndSection();
+
+}
+
+bool FCineSplineComponentVisualizer::HandleInputDelta(FEditorViewportClient* ViewportClient, FViewport* Viewport, FVector& DeltaTranslate, FRotator& DeltaRotate, FVector& DeltaScale)
+{
+	if (!DeltaRotate.IsZero() && bManipulatePointRotation)
+	{
+		return TransformPointRotation(DeltaRotate);
+	}
+	else
+	{
+		return FSplineComponentVisualizer::HandleInputDelta(ViewportClient, Viewport, DeltaTranslate, DeltaRotate, DeltaScale);
+	}
+}
+
+bool FCineSplineComponentVisualizer::TransformPointRotation(const FRotator& DeltaRotate)
+{
+	if (UCineSplineComponent* SplineComp = Cast<UCineSplineComponent>(GetEditedSplineComponent()))
+	{
+		check(SelectionState);
+
+		const int32 NumPoints = SplineComp->GetNumberOfSplinePoints();
+		const TSet<int32>& SelectedKeys = SelectionState->GetSelectedKeys();
+		int32 LastKeyIndexSelected = SelectionState->GetVerifiedLastKeyIndexSelected(NumPoints);
+
+		check(SelectedKeys.Num() > 0);
+		check(SelectedKeys.Contains(LastKeyIndexSelected));
+
+		SplineComp->Modify();
+
+		for (int32 SelectedKeyIndex : SelectedKeys)
+		{
+			check(SelectedKeyIndex >= 0);
+			check(SelectedKeyIndex < NumPoints);
+
+			const FQuat Rotation = SplineComp->GetPointRotationAtSplinePoint(SelectedKeyIndex);
+			SplineComp->SetPointRotationAtSplinePoint(SelectedKeyIndex, DeltaRotate.Quaternion() * Rotation);
+		}
+
+		return true;
+	
+	}
+	return false;
+}
+
+bool FCineSplineComponentVisualizer::GetCustomInputCoordinateSystem(const FEditorViewportClient* ViewportClient, FMatrix& OutMatrix) const
+{
+	if (bManipulatePointRotation && (ViewportClient->GetWidgetCoordSystemSpace() == COORD_Local || ViewportClient->GetWidgetMode() == UE::Widget::WM_Rotate))
+	{
+		OutMatrix = FRotationMatrix(FRotator());
+		return true;
+	}
+	return FSplineComponentVisualizer::GetCustomInputCoordinateSystem(ViewportClient, OutMatrix);
+}
+
+void FCineSplineComponentVisualizer::OnSetVisualizeNormalizedPosition()
+{
+	if (UCineSplineComponent* SplineComp = Cast<UCineSplineComponent>(GetEditedSplineComponent()))
+	{
+		SplineComp->bShouldVisualizeNormalizedPosition = !SplineComp->bShouldVisualizeNormalizedPosition;
+		GEditor->RedrawLevelEditingViewports(true);
+	}
+}
+
+bool FCineSplineComponentVisualizer::IsVisualizeNormalizedPosition() const
+{
+	if (UCineSplineComponent* SplineComp = Cast<UCineSplineComponent>(GetEditedSplineComponent()))
+	{
+		return SplineComp->bShouldVisualizeNormalizedPosition;
+	}
+	return false;
+}
+
+void FCineSplineComponentVisualizer::OnSetVisualizeSplineLength()
+{
+	if (UCineSplineComponent* SplineComp = Cast<UCineSplineComponent>(GetEditedSplineComponent()))
+	{
+		SplineComp->bShouldVisualizeSplineLength = !SplineComp->bShouldVisualizeSplineLength;
+		GEditor->RedrawLevelEditingViewports(true);
+	}
+}
+
+bool FCineSplineComponentVisualizer::IsVisualizeSplineLength() const
+{
+	if (UCineSplineComponent* SplineComp = Cast<UCineSplineComponent>(GetEditedSplineComponent()))
+	{
+		return SplineComp->bShouldVisualizeSplineLength;
+	}
+	return false;
+}
+
+void FCineSplineComponentVisualizer::OnSetVisualizePointRotation()
+{
+	if (UCineSplineComponent* SplineComp = Cast<UCineSplineComponent>(GetEditedSplineComponent()))
+	{
+		SplineComp->bShouldVisualizePointRotation = !SplineComp->bShouldVisualizePointRotation;
+		GEditor->RedrawLevelEditingViewports(true);
+	}
+}
+
+bool FCineSplineComponentVisualizer::IsVisualizePointRotation() const
+{
+	if (UCineSplineComponent* SplineComp = Cast<UCineSplineComponent>(GetEditedSplineComponent()))
+	{
+		return SplineComp->bShouldVisualizePointRotation;
+	}
+	return false;
+}
+
+void FCineSplineComponentVisualizer::OnSetManipulatePointRotation()
+{
+	bManipulatePointRotation = !bManipulatePointRotation;
+}
+
+bool FCineSplineComponentVisualizer::IsManipulatePointRotation() const
+{
+	return bManipulatePointRotation;
+}
+
 #undef LOCTEXT_NAMESPACE
\ No newline at end of file
diff --git a/Engine/Plugins/Experimental/CineCameraRigs/Source/CineCameraRigsEditor/Private/CineSplineComponentVisualizer.h b/Engine/Plugins/Experimental/CineCameraRigs/Source/CineCameraRigsEditor/Private/CineSplineComponentVisualizer.h
--- a/Engine/Plugins/Experimental/CineCameraRigs/Source/CineCameraRigsEditor/Private/CineSplineComponentVisualizer.h
+++ b/Engine/Plugins/Experimental/CineCameraRigs/Source/CineCameraRigsEditor/Private/CineSplineComponentVisualizer.h
@@ -13,12 +13,36 @@
 #include "Components/SplineComponent.h"
 #include "CineSplineComponent.h"
 
+class FUICommandList;
 
 class FCineSplineComponentVisualizer : public FSplineComponentVisualizer
 {
 public:
+	FCineSplineComponentVisualizer();
+	virtual ~FCineSplineComponentVisualizer();
+
 	//~ Begin FComponentVisualizer Interface
 	virtual void OnRegister() override;
 	virtual void DrawVisualizationHUD(const UActorComponent* Component, const FViewport* Viewport, const FSceneView* View, FCanvas* Canvas) override;
+	virtual void DrawVisualization(const UActorComponent* Component, const FSceneView* View, FPrimitiveDrawInterface* PDI) override;
+	virtual void GenerateContextMenuSections(FMenuBuilder& InMenuBuilder) const;
+	virtual bool HandleInputDelta(FEditorViewportClient* ViewportClient, FViewport* Viewport, FVector& DeltaTranslate, FRotator& DeltaRotate, FVector& DeltaScale) override;
+	virtual bool GetCustomInputCoordinateSystem(const FEditorViewportClient* ViewportClient, FMatrix& OutMatrix) const override;
+
+protected:
+	bool bManipulatePointRotation;
+
+	bool TransformPointRotation(const FRotator& DeltaRotate);
+
+	void OnSetVisualizeNormalizedPosition();
+	bool IsVisualizeNormalizedPosition() const;
+	void OnSetVisualizeSplineLength();
+	bool IsVisualizeSplineLength() const;
+	void OnSetVisualizePointRotation();
+	bool IsVisualizePointRotation() const;
+	void OnSetManipulatePointRotation();
+	bool IsManipulatePointRotation() const;
+
+	TSharedPtr<FUICommandList> CineSplineComponentVisualizerActions;
 
 };
\ No newline at end of file
