Add blueprint setter for bDisplaySplineHeatmap

diff --git a/Engine/Plugins/Experimental/CineCameraRigs/Source/CineCameraRigs/Private/CineCameraRigRail.cpp b/Engine/Plugins/Experimental/CineCameraRigs/Source/CineCameraRigs/Private/CineCameraRigRail.cpp
--- a/Engine/Plugins/Experimental/CineCameraRigs/Source/CineCameraRigs/Private/CineCameraRigRail.cpp
+++ b/Engine/Plugins/Experimental/CineCameraRigs/Source/CineCameraRigs/Private/CineCameraRigRail.cpp
@@ -417,6 +417,14 @@ void ACineCameraRigRail::SetDriveMode(ECineCameraRigRailDriveMode InMode)
 	SetMIDParameters();
 }
 
+void ACineCameraRigRail::SetDisplaySpeedHeatmap(bool bEnable)
+{
+	bDisplaySpeedHeatmap = bEnable;
+	UpdateSpeedHeatmap();
+	UpdateSplineMeshMID();
+	SetMIDParameters();
+}
+
 void ACineCameraRigRail::UpdateSpeedHeatmap()
 {
 	// If bDisplaySpeedHeatmap is false, assigns default grey texture
diff --git a/Engine/Plugins/Experimental/CineCameraRigs/Source/CineCameraRigs/Public/CineCameraRigRail.h b/Engine/Plugins/Experimental/CineCameraRigs/Source/CineCameraRigs/Public/CineCameraRigRail.h
--- a/Engine/Plugins/Experimental/CineCameraRigs/Source/CineCameraRigs/Public/CineCameraRigRail.h
+++ b/Engine/Plugins/Experimental/CineCameraRigs/Source/CineCameraRigs/Public/CineCameraRigRail.h
@@ -61,7 +61,7 @@ public:
 	TObjectPtr<UTexture2D> SplineMeshTexture;
 
 	/* Enable speed visualization*/
-	UPROPERTY(EditAnywhere, BlueprintReadWrite, Category = "SplineVisualization")
+	UPROPERTY(EditAnywhere, BlueprintSetter=SetDisplaySpeedHeatmap, Category = "SplineVisualization")
 	bool bDisplaySpeedHeatmap = true;
 
 	/* Number of speed samples per spline segment*/
@@ -136,6 +136,10 @@ public:
 	UFUNCTION(BlueprintSetter)
 	void SetDriveMode(ECineCameraRigRailDriveMode InMode);
 
+	/* Enable display speed heatmap*/
+	UFUNCTION(BlueprintSetter)
+	void SetDisplaySpeedHeatmap(bool bEnable);
+
 #if WITH_EDITOR
 	virtual void PostEditChangeProperty(FPropertyChangedEvent& PropertyChangedEvent) override;
 #endif
