Tag as beta version

diff --git a/Engine/Plugins/Experimental/CineCameraRigs/CineCameraRigs.uplugin b/Engine/Plugins/Experimental/CineCameraRigs/CineCameraRigs.uplugin
--- a/Engine/Plugins/Experimental/CineCameraRigs/CineCameraRigs.uplugin
+++ b/Engine/Plugins/Experimental/CineCameraRigs/CineCameraRigs.uplugin
@@ -12,7 +12,7 @@
 	"SupportURL": "",
 	"EnabledByDefault": false,
 	"CanContainContent": true,
-	"IsExperimentalVersion": true,
+	"IsBetaVersion": true,
 	"Installed": false,
 	"Modules": [
 		{
