Disable speed heatmap when CineCameraRigRail is driven by Sequencer.
Update speed progress to fix the jump

diff --git a/Engine/Plugins/Experimental/CineCameraRigs/CineCameraRigs.uplugin b/Engine/Plugins/Experimental/CineCameraRigs/CineCameraRigs.uplugin
--- a/Engine/Plugins/Experimental/CineCameraRigs/CineCameraRigs.uplugin
+++ b/Engine/Plugins/Experimental/CineCameraRigs/CineCameraRigs.uplugin
@@ -34,6 +34,14 @@
 		{
 			"Name": "ConcertSyncCore",
 			"Enabled": true
+		},
+		{
+			"Name": "SequencerScripting",
+			"Enabled": true
+		},
+		{
+			"Name": "LevelSequenceEditor",
+			"Enabled": true
 		}
 	]
 }
diff --git a/Engine/Plugins/Experimental/CineCameraRigs/Source/CineCameraRigs/CineCameraRigs.Build.cs b/Engine/Plugins/Experimental/CineCameraRigs/Source/CineCameraRigs/CineCameraRigs.Build.cs
--- a/Engine/Plugins/Experimental/CineCameraRigs/Source/CineCameraRigs/CineCameraRigs.Build.cs
+++ b/Engine/Plugins/Experimental/CineCameraRigs/Source/CineCameraRigs/CineCameraRigs.Build.cs
@@ -55,5 +55,12 @@ public class CineCameraRigs : ModuleRules
 			}
 			);
 
+		if (Target.bBuildEditor == true)
+		{
+			PrivateDependencyModuleNames.Add("LevelSequenceEditor");
+			PrivateDependencyModuleNames.Add("SequencerScriptingEditor");
+			PrivateDependencyModuleNames.Add("UnrealEd");
+		}
+
 	}
 }
diff --git a/Engine/Plugins/Experimental/CineCameraRigs/Source/CineCameraRigs/Private/CineCameraRigRail.cpp b/Engine/Plugins/Experimental/CineCameraRigs/Source/CineCameraRigs/Private/CineCameraRigRail.cpp
--- a/Engine/Plugins/Experimental/CineCameraRigs/Source/CineCameraRigs/Private/CineCameraRigRail.cpp
+++ b/Engine/Plugins/Experimental/CineCameraRigs/Source/CineCameraRigs/Private/CineCameraRigRail.cpp
@@ -8,11 +8,28 @@
 #include "CineSplineComponent.h"
 #include "Components/SplineMeshComponent.h"
 
+#include "Algo/AnyOf.h"
 #include "Engine/World.h"
 #include "Kismet/KismetMathLibrary.h"
 #include "Materials/MaterialInstanceDynamic.h"
 #include "UObject/Package.h"
 
+#if WITH_EDITOR
+#include "Editor.h"
+#include "LevelSequence.h"
+#include "LevelSequenceEditorBlueprintLibrary.h"
+#include "MovieScene.h"
+#include "MovieSceneBinding.h"
+#include "MovieSceneSequence.h"
+#include "MovieSceneTrack.h"
+#include "Channels/MovieSceneFloatChannel.h"
+#include "Sections/MovieSceneFloatSection.h"
+#include "Sections/MovieSceneSubSection.h"
+#include "SequencerTools.h"
+#include "Tracks/MovieSceneFloatTrack.h"
+#include "Tracks/MovieSceneSubTrack.h"
+#endif
+
 ACineCameraRigRail::ACineCameraRigRail(const FObjectInitializer& ObjectInitializer)
 	: Super(ObjectInitializer
 		.SetDefaultSubobjectClass<UCineSplineComponent>(TEXT("RailSplineComponent"))
@@ -44,7 +61,13 @@ ACineCameraRigRail::ACineCameraRigRail(const FObjectInitializer& ObjectInitializ
 	{
 		PreviewMesh_Mount->bHiddenInGame = false;
 	}
+#endif
 
+#if WITH_EDITOR
+	if (!IsTemplate() && GEditor)
+	{
+		GEditor->GetTimerManager()->SetTimer(SequencerCheckHandle, this, &ACineCameraRigRail::OnSequencerCheck, 1.0, true);
+	}
 #endif
 }
 
@@ -238,24 +261,17 @@ void ACineCameraRigRail::PostEditChangeProperty(FPropertyChangedEvent& PropertyC
 	else if (PropertyName == GET_MEMBER_NAME_CHECKED(ACineCameraRigRail, AbsolutePositionOnRail) ||
 			 PropertyName == GET_MEMBER_NAME_CHECKED(ACineCameraRigRail, CurrentPositionOnRail))
 	{
-		SpeedAccumulatedTime = 0.0f;
-		if (Speed != 0.0f)
-		{
-			float TotalTime = CineSplineComponent->GetSplineLength() / Speed;
-			SpeedAccumulatedTime = UKismetMathLibrary::SafeDivide(AbsolutePositionOnRail - StartPositionValue(), LastPositionValue() - StartPositionValue()) * TotalTime;
-		}
+		UpdateSpeedProgress();
 	}
-	else if (PropertyName == GET_MEMBER_NAME_CHECKED(ACineCameraRigRail, Speed))
+	else if (PropertyName == GET_MEMBER_NAME_CHECKED(ACineCameraRigRail, DriveMode))
 	{
-		if (Speed != 0.0)
-		{
-			float TotalTime = CineSplineComponent->GetSplineLength() / Speed;
-			SpeedAccumulatedTime = TotalTime * SpeedProgress;
-		}
+		UpdateSpeedProgress();
+		UpdateSpeedHeatmap();
+		UpdateSplineMeshMID();
+		SetMIDParameters();
 	}
 	else if (PropertyName == GET_MEMBER_NAME_CHECKED(ACineCameraRigRail, bDisplaySpeedHeatmap) ||
 		     PropertyName == GET_MEMBER_NAME_CHECKED(ACineCameraRigRail, SpeedSampleCountPerSegment) || 
-			 PropertyName == GET_MEMBER_NAME_CHECKED(ACineCameraRigRail, DriveMode) || 
 			 PropertyName == GET_MEMBER_NAME_CHECKED(ACineCameraRigRail, bUseAbsolutePosition))
 	{
 		UpdateSpeedHeatmap();
@@ -347,6 +363,16 @@ void ACineCameraRigRail::Tick(float DeltaTime)
 		DriveBySpeed(DeltaTime);
 	}
 }
+void ACineCameraRigRail::PostLoad()
+{
+	Super::PostLoad();
+	if (CineSplineComponent)
+	{
+		UpdateSpeedHeatmap();
+		UpdateSplineMeshMID();
+		SetMIDParameters();
+	}
+}
 
 void ACineCameraRigRail::DriveByParam(float DeltaTime)
 {
@@ -372,13 +398,13 @@ void ACineCameraRigRail::DriveBySpeed(float DeltaTime)
 		return;
 	}
 
-	float TotalTime = CineSplineComponent->GetSplineLength() / Speed;
-
-	SpeedAccumulatedTime += DeltaTime;
-	SpeedAccumulatedTime = bLoop ? FMath::Fmod(SpeedAccumulatedTime, TotalTime) : FMath::Clamp(SpeedAccumulatedTime, 0.0f, TotalTime);
-
-	SpeedProgress = SpeedAccumulatedTime / TotalTime;
+	float TotalTime = CineSplineComponent->GetSplineLength() / FMath::Abs(Speed);
+	float CurrentTime = TotalTime * SpeedProgress;
+	CurrentTime += (FMath::Sign(Speed) * DeltaTime);
+	CurrentTime = bLoop ? FMath::Fmod(CurrentTime + TotalTime, TotalTime) : FMath::Clamp(CurrentTime, 0.0f, TotalTime);
 
+	SpeedProgress = CurrentTime / TotalTime;
+	
 	if (bUseAbsolutePosition)
 	{
 		float Distance = CineSplineComponent->GetSplineLength() * SpeedProgress;
@@ -402,16 +428,10 @@ float ACineCameraRigRail::LastPositionValue() const
 	return bUseAbsolutePosition ? CineSplineComponent->GetFloatPropertyAtSplinePoint(NumPoints - 1, FName(TEXT("AbsolutePosition"))) : 1.0f;
 }
 
-void ACineCameraRigRail::SetDriveModeSpeed(float Value)
-{
-	float TotalTime = UKismetMathLibrary::SafeDivide(CineSplineComponent->GetSplineLength(), Value);
-	SpeedAccumulatedTime = TotalTime * SpeedProgress;
-	Speed = Value;
-}
-
 void ACineCameraRigRail::SetDriveMode(ECineCameraRigRailDriveMode InMode)
 {
 	DriveMode = InMode;
+	UpdateSpeedProgress();
 	UpdateSpeedHeatmap();
 	UpdateSplineMeshMID();
 	SetMIDParameters();
@@ -425,13 +445,28 @@ void ACineCameraRigRail::SetDisplaySpeedHeatmap(bool bEnable)
 	SetMIDParameters();
 }
 
+void ACineCameraRigRail::UpdateSpeedProgress()
+{
+	const float Length = CineSplineComponent->GetSplineLength();
+	if (bUseAbsolutePosition)
+	{
+		const float InputKey = CineSplineComponent->GetInputKeyAtPosition(AbsolutePositionOnRail);
+		const float Distance = CineSplineComponent->GetDistanceAlongSplineAtSplineInputKey(InputKey);
+		SpeedProgress = UKismetMathLibrary::SafeDivide(Distance, Length);
+	}
+	else
+	{
+		SpeedProgress = CurrentPositionOnRail;
+	}
+}
+
 void ACineCameraRigRail::UpdateSpeedHeatmap()
 {
 	// If bDisplaySpeedHeatmap is false, assigns default grey texture
 	// If DriveMode is Speed, assigns green texture indicating it's constant speed
 	// Otherwise, samples the velocity along the spline and generates a heatmap texture
 
-	if (!bDisplaySpeedHeatmap)
+	if (!bDisplaySpeedHeatmap || IsSequencerDriven())
 	{
 		UTexture2D* DefaultTexture = GetDefault<UCineCameraRigRailSettings>()->DefaultSplineMeshTexture.LoadSynchronous();
 		if (DefaultTexture)
@@ -499,4 +534,95 @@ void ACineCameraRigRail::UpdateSpeedHeatmap()
 	float LowValue = AverageValue * 0.5f;
 	float HighValue = AverageValue * 2.0f;
 	UCineCameraRigRailHelpers::CreateOrUpdateSplineHeatmapTexture(static_cast<UTexture2D*&>(SplineMeshTexture), SpeedValues, LowValue, AverageValue, HighValue);
-}
\ No newline at end of file
+}
+
+bool ACineCameraRigRail::IsSequencerDriven()
+{
+#if WITH_EDITOR
+	return bSequencerDriven;
+#endif
+	return false;
+}
+
+#if WITH_EDITOR
+UMovieSceneFloatTrack* ACineCameraRigRail::FindPositionTrack(const UMovieSceneSequence* InSequence)
+{
+	if (!IsValid(InSequence) || !InSequence->GetMovieScene())
+	{
+		return nullptr;
+	}
+	
+	UMovieScene* MovieScene = InSequence->GetMovieScene();
+
+	const TArray<FMovieSceneBinding>& Bindings = MovieScene->GetBindings();
+	const FName PropertyName = bUseAbsolutePosition ? FName(TEXT("AbsolutePositionOnRail")) : FName(TEXT("CurrentPositionOnRail"));
+	for (const FMovieSceneBinding& Binding : Bindings)
+	{
+		TArray<UObject*, TInlineAllocator<1>> BoundObjects;
+		InSequence->LocateBoundObjects(Binding.GetObjectGuid(), nullptr, BoundObjects);
+		if (BoundObjects.IsEmpty() || BoundObjects[0] != this)
+		{
+			continue;
+		}
+		if (UMovieSceneFloatTrack* Track = MovieScene->FindTrack<UMovieSceneFloatTrack>(Binding.GetObjectGuid(), PropertyName))
+		{
+			return Track;
+		}
+	}
+
+	// Find in subsequences
+	for (const UMovieSceneTrack* Track : MovieScene->GetTracks())
+	{
+		const UMovieSceneSubTrack* SubTrack = Cast<UMovieSceneSubTrack>(Track);
+		if (!SubTrack || SubTrack->IsEvalDisabled())
+		{
+			continue;
+		}
+
+		for( const UMovieSceneSection* Section : SubTrack->GetAllSections())
+		{
+			const UMovieSceneSubSection* SubSection = Cast<UMovieSceneSubSection>(Section);
+			if (!SubSection)
+			{
+				continue;
+			}
+
+			UMovieSceneSequence* SubSequence = SubSection->GetSequence();
+			if (!SubSequence)
+			{
+				continue;
+			}
+
+			if (UMovieSceneFloatTrack* SubSequenceTrack = FindPositionTrack(SubSequence))
+			{
+				return SubSequenceTrack;
+			}
+		}
+	}
+
+	return nullptr;
+}
+
+void ACineCameraRigRail::OnSequencerCheck()
+{
+	bool bHasKey = false;
+	ULevelSequence* LevelSequence = ULevelSequenceEditorBlueprintLibrary::GetCurrentLevelSequence();
+	UMovieSceneFloatTrack* Track = FindPositionTrack(LevelSequence);
+	if (IsValid(Track) && !Track->IsEvalDisabled())
+	{
+		bHasKey = Algo::AnyOf(Track->GetAllSections(), [](const UMovieSceneSection* Section)
+		{
+			const FMovieSceneFloatChannel* Channel = Section->GetChannelProxy().GetChannels<FMovieSceneFloatChannel>()[0];
+			return Channel->GetNumKeys() > 0;
+		});
+	}
+
+	if(bSequencerDriven != bHasKey)
+	{
+		bSequencerDriven = !bSequencerDriven;
+		UpdateSpeedHeatmap();
+		UpdateSplineMeshMID();
+		SetMIDParameters();
+	}
+}
+#endif
diff --git a/Engine/Plugins/Experimental/CineCameraRigs/Source/CineCameraRigs/Public/CineCameraRigRail.h b/Engine/Plugins/Experimental/CineCameraRigs/Source/CineCameraRigs/Public/CineCameraRigRail.h
--- a/Engine/Plugins/Experimental/CineCameraRigs/Source/CineCameraRigs/Public/CineCameraRigRail.h
+++ b/Engine/Plugins/Experimental/CineCameraRigs/Source/CineCameraRigs/Public/CineCameraRigRail.h
@@ -2,6 +2,7 @@
 #pragma once
 
 #include "CoreMinimal.h"
+#include "Engine/TimerHandle.h"
 #include "UObject/ObjectMacros.h"
 
 #include "CameraRig_Rail.h"
@@ -9,6 +10,9 @@
 
 #include "CineCameraRigRail.generated.h"
 
+class UMovieSceneSequence;
+class UMovieSceneFloatTrack;
+
 UENUM(BlueprintType)
 enum class ECineCameraRigRailDriveMode : uint8
 {
@@ -31,6 +35,7 @@ public:
 	ACineCameraRigRail(const FObjectInitializer& ObjectInitializer);
 
 	virtual void Tick(float DeltaTime) override;
+	virtual void PostLoad() override;
 
 	/* Returns CineSplineComponent*/
 	UFUNCTION(BlueprintPure, Category = "Rail Components")
@@ -45,7 +50,7 @@ public:
 	float AbsolutePositionOnRail = 1.0f;
 
 	/* Use PointRotation metadata for attachment orientation. If false, attachment orientation is based on the spline curvature*/
-	UPROPERTY(EditAnywhere, BlueprintReadWrite, Interp, Category = "Rail Controls", meta = (EditCondition = "bLockOrientationToRail"))
+	UPROPERTY(EditAnywhere, BlueprintReadWrite, Category = "Rail Controls", meta = (EditCondition = "bLockOrientationToRail"))
 	bool bUsePointRotation = true;
 
 	/* Material assigned to spline component mesh*/
@@ -60,8 +65,8 @@ public:
 	UPROPERTY(BlueprintSetter=SetSplineMeshTexture, Category = "SplineVisualization")
 	TObjectPtr<UTexture2D> SplineMeshTexture;
 
-	/* Enable speed visualization*/
-	UPROPERTY(EditAnywhere, BlueprintSetter=SetDisplaySpeedHeatmap, Category = "SplineVisualization")
+	/* Enable speed visualization. Automatically disabled when position property is driven in Sequencer*/
+	UPROPERTY(EditAnywhere, BlueprintSetter=SetDisplaySpeedHeatmap, Category = "SplineVisualization", meta=(EditCondition="!IsSequencerDriven()"))
 	bool bDisplaySpeedHeatmap = true;
 
 	/* Number of speed samples per spline segment*/
@@ -109,7 +114,7 @@ public:
 	ECineCameraRigRailDriveMode DriveMode = ECineCameraRigRailDriveMode::Manual;
 
 	/* Specifies the drive speed of the rig rail in centimeter per second */
-	UPROPERTY(EditAnywhere, BlueprintSetter=SetDriveModeSpeed, Category = "DriveMode")
+	UPROPERTY(EditAnywhere, Interp, Category = "DriveMode")
 	float Speed = 100;
 
 	/* Enable loop in speed or duration mode */
@@ -128,10 +133,6 @@ public:
 	UFUNCTION(BlueprintCallable, Category = "CineCameraRigRail")
 	FVector GetVelocityAtPosition(const float InPosition, const float delta = 0.001) const;
 
-	/* Set drive mode speed value*/
-	UFUNCTION(BlueprintSetter)
-	void SetDriveModeSpeed(float Value);
-
 	/* Set drive mode*/
 	UFUNCTION(BlueprintSetter)
 	void SetDriveMode(ECineCameraRigRailDriveMode InMode);
@@ -140,6 +141,10 @@ public:
 	UFUNCTION(BlueprintSetter)
 	void SetDisplaySpeedHeatmap(bool bEnable);
 
+	/* Returns true if the rig rail is driven by Sequencer */
+	UFUNCTION()
+	bool IsSequencerDriven();
+
 #if WITH_EDITOR
 	virtual void PostEditChangeProperty(FPropertyChangedEvent& PropertyChangedEvent) override;
 #endif
@@ -161,5 +166,19 @@ protected:
 	float StartPositionValue() const;
 	float LastPositionValue() const;
 	float SpeedProgress = 0.0f;
-	float SpeedAccumulatedTime = 0.0f;
+	void UpdateSpeedProgress();
+
+
+#if WITH_EDITOR
+	UMovieSceneFloatTrack* FindPositionTrack(const UMovieSceneSequence* InSequence);
+
+	FTimerHandle SequencerCheckHandle;
+
+	/* Check if the rig rail is driven by Sequencer*/
+	UFUNCTION()
+	void OnSequencerCheck();
+
+	bool bSequencerDriven = false;
+#endif
+
 };
