Add CineCameraRigRail to Cinematic category

diff --git a/Engine/Plugins/Experimental/CineCameraRigs/Source/CineCameraRigsEditor/CineCameraRigsEditor.Build.cs b/Engine/Plugins/Experimental/CineCameraRigs/Source/CineCameraRigsEditor/CineCameraRigsEditor.Build.cs
--- a/Engine/Plugins/Experimental/CineCameraRigs/Source/CineCameraRigsEditor/CineCameraRigsEditor.Build.cs
+++ b/Engine/Plugins/Experimental/CineCameraRigs/Source/CineCameraRigsEditor/CineCameraRigsEditor.Build.cs
@@ -18,6 +18,7 @@ public class CineCameraRigsEditor : ModuleRules
                 "ComponentVisualizers",
 				"DetailCustomizations",
 				"Engine",
+				"PlacementMode",
 				"SlateCore",
 				"Slate",
 				"UnrealEd",
diff --git a/Engine/Plugins/Experimental/CineCameraRigs/Source/CineCameraRigsEditor/Private/CineCameraRigsEditor.cpp b/Engine/Plugins/Experimental/CineCameraRigs/Source/CineCameraRigsEditor/Private/CineCameraRigsEditor.cpp
--- a/Engine/Plugins/Experimental/CineCameraRigs/Source/CineCameraRigsEditor/Private/CineCameraRigsEditor.cpp
+++ b/Engine/Plugins/Experimental/CineCameraRigs/Source/CineCameraRigsEditor/Private/CineCameraRigsEditor.cpp
@@ -8,6 +8,7 @@
 #include "CineCameraRigRailDetails.h"
 #include "Editor.h"
 #include "Editor/UnrealEdEngine.h"
+#include "IPlacementModeModule.h"
 #include "UnrealEdGlobals.h"
 #include "PropertyEditorModule.h"
 #include "PropertyEditorDelegates.h"
@@ -28,6 +29,7 @@ public:
 			GUnrealEd->RegisterComponentVisualizer(CineSplineComponentName, Visualizer);
 			Visualizer->OnRegister();
 			RegisterCustomizations();
+			RegisterPlacementModeItems();
 		}
 	}
 
@@ -54,6 +56,36 @@ private:
 		FPropertyEditorModule& PropertyModule = FModuleManager::GetModuleChecked<FPropertyEditorModule>("PropertyEditor");
 		PropertyModule.UnregisterCustomPropertyTypeLayout(ACineCameraRigRail::StaticClass()->GetFName());
 	}
+	void RegisterPlacementModeItems()
+	{
+		if (const FPlacementCategoryInfo* Info = GetCinematicPlacementCategoryInfo())
+		{
+			IPlacementModeModule::Get().RegisterPlaceableItem(Info->UniqueHandle, MakeShared<FPlaceableItem>(
+				*ACineCameraRigRail::StaticClass(),
+				FAssetData(ACineCameraRigRail::StaticClass())));
+		}
+	}
+
+	const FPlacementCategoryInfo* GetCinematicPlacementCategoryInfo() const
+	{
+		IPlacementModeModule& PlacmentModeModule = IPlacementModeModule::Get();
+
+		if (const FPlacementCategoryInfo* RegisteredInfo = PlacmentModeModule.GetRegisteredPlacementCategory("Cinematic"))
+		{
+			return RegisteredInfo;
+		}
+
+		FPlacementCategoryInfo Info(
+			LOCTEXT("CinematicCategoryName", "Cinematic"),
+			FSlateIcon(FAppStyle::GetAppStyleSetName(), "PlacementBrowser.Icons.Cinematics"),
+			"Cinematic",
+			TEXT("PMCinematic"),
+			25
+		);
+
+		PlacmentModeModule.RegisterPlacementCategory(Info);
+		return PlacmentModeModule.GetRegisteredPlacementCategory("Cinematic");
+	}
 };
 
 
