Creates 2 spline points by default to mimic original rig rail. SplineModifier deletes the 2nd point on spawn to keep existing behavior

diff --git a/Engine/Plugins/Experimental/CineCameraRigs/Source/CineCameraRigs/Private/CineSplineComponent.cpp b/Engine/Plugins/Experimental/CineCameraRigs/Source/CineCameraRigs/Private/CineSplineComponent.cpp
--- a/Engine/Plugins/Experimental/CineCameraRigs/Source/CineCameraRigs/Private/CineSplineComponent.cpp
+++ b/Engine/Plugins/Experimental/CineCameraRigs/Source/CineCameraRigs/Private/CineSplineComponent.cpp
@@ -22,6 +22,9 @@ UCineSplineComponent::UCineSplineComponent(const FObjectInitializer& ObjectIniti
 	Data.FocusDistance = Defaults.DefaultFocusDistance;
 
 	AddSplineDataAtPosition(1.0f, Data);
+
+	Data.Location.X = 100.0;
+	AddSplineDataAtPosition(2.0f, Data);
 }
 
 
