CineCameraRigs: Removed include from private directory

diff --git a/Engine/Plugins/Experimental/CineCameraRigs/Source/CineCameraRigsEditor/CineCameraRigsEditor.Build.cs b/Engine/Plugins/Experimental/CineCameraRigs/Source/CineCameraRigsEditor/CineCameraRigsEditor.Build.cs
--- a/Engine/Plugins/Experimental/CineCameraRigs/Source/CineCameraRigsEditor/CineCameraRigsEditor.Build.cs
+++ b/Engine/Plugins/Experimental/CineCameraRigs/Source/CineCameraRigsEditor/CineCameraRigsEditor.Build.cs
@@ -7,8 +7,6 @@ public class CineCameraRigsEditor : ModuleRules
 	public CineCameraRigsEditor(ReadOnlyTargetRules Target) : base(Target)
 	{
 		PCHUsage = ModuleRules.PCHUsageMode.UseExplicitOrSharedPCHs;
-		
-        PrivateIncludePaths.Add("CineCameraRigs/Private");
 
 		PrivateDependencyModuleNames.AddRange(
 			new string[] {
