VCam: Parenting and attach rig prototype

#virtualized

diff --git a/Engine/Plugins/Experimental/CineCameraRigs/Source/CineCameraRigs/CineCameraRigs.Build.cs b/Engine/Plugins/Experimental/CineCameraRigs/Source/CineCameraRigs/CineCameraRigs.Build.cs
--- a/Engine/Plugins/Experimental/CineCameraRigs/Source/CineCameraRigs/CineCameraRigs.Build.cs
+++ b/Engine/Plugins/Experimental/CineCameraRigs/Source/CineCameraRigs/CineCameraRigs.Build.cs
@@ -27,6 +27,8 @@ public class CineCameraRigs : ModuleRules
 			{
 				"Core",
 				"CinematicCamera",
+				"AnimationCore",
+				"Constraints"
 				// ... add other public dependencies that you statically link with here ...
 			}
 			);
diff --git a/Engine/Plugins/Experimental/CineCameraRigs/Source/CineCameraRigs/Private/CineCameraAttachMount.cpp b/Engine/Plugins/Experimental/CineCameraRigs/Source/CineCameraRigs/Private/CineCameraAttachMount.cpp
new file mode 100644
--- /dev/null
+++ b/Engine/Plugins/Experimental/CineCameraRigs/Source/CineCameraRigs/Private/CineCameraAttachMount.cpp
@@ -0,0 +1,309 @@
+// Copyright Epic Games, Inc. All Rights Reserved.
+
+#include "CineCameraAttachMount.h"
+
+#include "TransformConstraint.h"
+#include "TransformableHandle.h"
+
+#include "Components/BillboardComponent.h"
+#include "Components/StaticMeshComponent.h"
+#include "GameFramework/SpringArmComponent.h"
+
+#include "UObject/ConstructorHelpers.h"
+#include "Engine/Texture2D.h"
+
+#if WITH_EDITOR
+#include "Editor.h"
+#include "UnrealEdGlobals.h"
+#include "Editor/UnrealEdEngine.h"
+#include "LevelEditorViewport.h"
+#include "SLevelViewport.h"
+#include "SEditorViewport.h"
+#endif
+
+ACineCameraAttachMount::ACineCameraAttachMount(const FObjectInitializer& ObjectInitializer)
+	: Super(ObjectInitializer)
+{
+	PrimaryActorTick.bCanEverTick = true;
+
+#if WITH_EDITORONLY_DATA
+	bShowPreviewMeshes = true;
+	PreviewMeshScale = 1.f;
+#endif
+
+	USceneComponent* SceneRoot = CreateDefaultSubobject<USceneComponent>("Root");
+	SetRootComponent(SceneRoot);
+
+	SpringArmComponent = CreateDefaultSubobject<USpringArmComponent>(TEXT("SpringArm"));
+	SpringArmComponent->AttachToComponent(SceneRoot, FAttachmentTransformRules::KeepRelativeTransform);
+	SpringArmComponent->bDoCollisionTest = false;
+	SpringArmComponent->bEnableCameraLag = true;
+	SpringArmComponent->bEnableCameraRotationLag = true;
+	SpringArmComponent->TargetArmLength = 0;
+
+#if WITH_EDITOR
+	CreatePreviewMesh();
+#endif
+}
+
+void ACineCameraAttachMount::Tick(float DeltaTime)
+{
+	Super::Tick(DeltaTime);
+#if WITH_EDITOR
+	if (bConstraintUpdated)
+	{
+		GUnrealEd->UpdatePivotLocationForSelection();
+		bConstraintUpdated = false;
+	}
+	UpdatePreviewMeshes();
+#endif
+}
+
+bool ACineCameraAttachMount::ShouldTickIfViewportsOnly() const
+{
+	return true;
+}
+
+USceneComponent* ACineCameraAttachMount::GetDefaultAttachComponent() const
+{
+	return SpringArmComponent;
+}
+
+void ACineCameraAttachMount::PostEditChangeProperty(FPropertyChangedEvent& PropertyChangedEvent)
+{
+	Super::PostEditChangeProperty(PropertyChangedEvent);
+	const FName PropertyName = PropertyChangedEvent.GetPropertyName();
+	if (PropertyName == GET_MEMBER_NAME_CHECKED(ACineCameraAttachMount, TargetActor) ||
+		PropertyName == GET_MEMBER_NAME_CHECKED(ACineCameraAttachMount, TargetSocket))
+	{
+		CreateConstraint();
+	}
+	else if (PropertyName == GET_MEMBER_NAME_CHECKED(ACineCameraAttachMount, bEnableLocationLag))
+	{
+		SpringArmComponent->bEnableCameraLag = bEnableLocationLag;
+	}
+	else if (PropertyName == GET_MEMBER_NAME_CHECKED(ACineCameraAttachMount, bEnableRotationLag))
+	{
+		SpringArmComponent->bEnableCameraRotationLag = bEnableRotationLag;
+	}
+	else if (PropertyName == GET_MEMBER_NAME_CHECKED(ACineCameraAttachMount, LocationLagSpeed))
+	{
+		SpringArmComponent->CameraLagSpeed = LocationLagSpeed;
+	}
+	else if (PropertyName == GET_MEMBER_NAME_CHECKED(ACineCameraAttachMount, RotationLagSpeed))
+	{
+		SpringArmComponent->CameraRotationLagSpeed = RotationLagSpeed;
+	}
+}
+
+void ACineCameraAttachMount::PostEditChangeChainProperty(FPropertyChangedChainEvent& PropertyChangedEvent)
+{
+	Super::PostEditChangeChainProperty(PropertyChangedEvent);
+	const FName StructName = PropertyChangedEvent.PropertyChain.GetActiveMemberNode()->GetValue()->GetFName();
+	if (StructName == GET_MEMBER_NAME_CHECKED(ACineCameraAttachMount, TransformFilter))
+	{
+		SetAxisFilter();
+	}
+}
+
+void ACineCameraAttachMount::SetAxisFilter()
+{
+	if (UTickableParentConstraint* Constraint = GetConstraint())
+	{
+		Constraint->TransformFilter = TransformFilter;
+	}
+}
+
+void ACineCameraAttachMount::CreateConstraint()
+{
+	UWorld* World = GetWorld();
+	TArray< TWeakObjectPtr<UTickableConstraint> > Constraints;
+	FTransformConstraintUtils::GetParentConstraints(World, this, Constraints);
+
+	for (const TWeakObjectPtr<UTickableConstraint>& Constraint : Constraints)
+	{
+		FConstraintsManagerController& Controller = FConstraintsManagerController::Get(World);
+		Controller.RemoveConstraint(Constraint.Get());
+	}
+
+	if (!TargetActor.IsValid())
+	{
+		return;
+	}
+
+	UTickableTransformConstraint* Constraint = FTransformConstraintUtils::CreateAndAddFromObjects(World, TargetActor.Get(), TargetSocket, this, NAME_None, ETransformConstraintType::Parent);
+	if (Constraint)
+	{
+		Constraint->bDynamicOffset = true;
+		Cast<UTickableParentConstraint>(Constraint)->OffsetTransform = FTransform::Identity;
+		FConstraintsManagerController& Controller = FConstraintsManagerController::Get(World);
+		Controller.StaticConstraintCreated(World, Constraint);
+	}
+	bConstraintUpdated = true;
+}
+
+UTickableParentConstraint* ACineCameraAttachMount::GetConstraint()
+{
+	UWorld* World = GetWorld();
+	TArray< TWeakObjectPtr<UTickableConstraint> > Constraints;
+	FTransformConstraintUtils::GetParentConstraints(World, this, Constraints);
+	if (Constraints.Num() == 0)
+	{
+		return nullptr;
+	}
+	return Cast<UTickableParentConstraint>(Constraints[0].Get());
+}
+
+void ACineCameraAttachMount::ResetLocationOffset()
+{
+	if (UTickableParentConstraint* Constraint = GetConstraint())
+	{
+		Constraint->OffsetTransform.SetLocation(FVector::ZeroVector);
+		bConstraintUpdated = true;
+	}
+}
+
+void ACineCameraAttachMount::ResetRotationOffset()
+{
+	if (UTickableParentConstraint* Constraint = GetConstraint())
+	{
+		Constraint->OffsetTransform.SetRotation(FQuat::Identity);
+		bConstraintUpdated = true;
+	}
+}
+
+void ACineCameraAttachMount::ZeroRoll()
+{
+	if (UTickableParentConstraint* Constraint = GetConstraint())
+	{
+		FTransform ParentTransform = Constraint->GetParentGlobalTransform();
+		FTransform ChildTransform = Constraint->GetChildGlobalTransform();
+		FRotator WorldRot = ChildTransform.Rotator();
+		WorldRot.Roll = 0;
+		ChildTransform.SetRotation(FQuat(WorldRot));
+		Constraint->OffsetTransform = ChildTransform.GetRelativeTransform(ParentTransform);
+		bConstraintUpdated = true;
+	}
+	else
+	{
+		FRotator WorldRot = RootComponent->GetComponentRotation();
+		WorldRot.Roll = 0;
+		RootComponent->SetWorldRotation(WorldRot);
+	}
+}
+
+void ACineCameraAttachMount::ZeroTilt()
+{
+	if (UTickableParentConstraint* Constraint = GetConstraint())
+	{
+		FTransform ParentTransform = Constraint->GetParentGlobalTransform();
+		FTransform ChildTransform = Constraint->GetChildGlobalTransform();
+		FRotator WorldRot = ChildTransform.Rotator();
+		WorldRot.Pitch = 0;
+		ChildTransform.SetRotation(FQuat(WorldRot));
+		Constraint->OffsetTransform = ChildTransform.GetRelativeTransform(ParentTransform);
+		bConstraintUpdated = true;
+	}
+	else
+	{
+		FRotator WorldRot = RootComponent->GetComponentRotation();
+		WorldRot.Pitch = 0;
+		RootComponent->SetWorldRotation(WorldRot);
+	}
+}
+
+void ACineCameraAttachMount::SetEnableLocationLag(bool bEnabled)
+{
+	SpringArmComponent->bEnableCameraLag = bEnableLocationLag = bEnabled;
+}
+
+void ACineCameraAttachMount::SetEnableRotationLag(bool bEnabled)
+{
+	SpringArmComponent->bEnableCameraRotationLag = bEnableRotationLag = bEnabled;
+}
+
+void ACineCameraAttachMount::SetLocationLagSpeed(float Speed)
+{
+	SpringArmComponent->CameraLagSpeed = LocationLagSpeed = Speed;
+}
+
+void ACineCameraAttachMount::SetRotationLagSpeed(float Speed)
+{
+	SpringArmComponent->CameraRotationLagSpeed = RotationLagSpeed = Speed;
+}
+
+#if WITH_EDITOR
+
+void ACineCameraAttachMount::CreatePreviewMesh()
+{
+	static ConstructorHelpers::FObjectFinder<UStaticMesh> PreviewMesh(TEXT("/CineCameraRigs/AttachMount/SM_CineCameraAttachMount_PreviewMesh.SM_CineCameraAttachMount_PreviewMesh"));
+	PreviewMesh_Root = CreateOptionalDefaultSubobject<UStaticMeshComponent>(TEXT("PreviewMesh_Root"));
+	if (PreviewMesh_Root)
+	{
+		PreviewMesh_Root->SetStaticMesh(PreviewMesh.Object);
+		PreviewMesh_Root->SetIsVisualizationComponent(true);
+		PreviewMesh_Root->SetVisibility(bShowPreviewMeshes);
+		PreviewMesh_Root->SetCollisionProfileName(UCollisionProfile::NoCollision_ProfileName);
+		PreviewMesh_Root->bHiddenInGame = true;
+		PreviewMesh_Root->CastShadow = false;
+		PreviewMesh_Root->SetWorldScale3D(FVector(PreviewMeshScale, PreviewMeshScale, PreviewMeshScale));
+		PreviewMesh_Root->SetupAttachment(RootComponent);
+	}
+
+	static ConstructorHelpers::FObjectFinder<UStaticMesh> MountMeshData(TEXT("/Engine/EditorMeshes/Camera/SM_RailRig_Mount.SM_RailRig_Mount"));
+	MountMesh = MountMeshData.Object;
+}
+
+void ACineCameraAttachMount::UpdatePreviewMeshes()
+{
+	if (PreviewMesh_Root)
+	{
+		PreviewMesh_Root->SetVisibility(bShowPreviewMeshes);
+		PreviewMesh_Root->SetWorldScale3D(FVector(PreviewMeshScale, PreviewMeshScale, PreviewMeshScale));
+	}
+
+	TArray< AActor* > AttachedActors;
+	GetAttachedActors(AttachedActors);
+
+	if (PreviewMeshes_Mount.Num() != AttachedActors.Num())
+	{
+		for (UStaticMeshComponent* PreviewMesh : PreviewMeshes_Mount)
+		{
+			PreviewMesh->UnregisterComponent();
+			PreviewMesh->Modify();
+			PreviewMesh->DestroyComponent();
+		}
+		PreviewMeshes_Mount.Empty();
+
+		for (int Index = 0; Index < AttachedActors.Num(); ++Index)
+		{
+			UStaticMeshComponent* const PreviewMesh = NewObject<UStaticMeshComponent>(this);
+			if (PreviewMesh)
+			{
+				PreviewMesh->SetStaticMesh(MountMesh);
+				PreviewMesh->SetIsVisualizationComponent(true);
+				PreviewMesh->SetVisibility(bShowPreviewMeshes);
+				PreviewMesh->SetCollisionProfileName(UCollisionProfile::NoCollision_ProfileName);
+				PreviewMesh->bHiddenInGame = true;
+				PreviewMesh->CastShadow = false;
+				PreviewMesh->SetWorldScale3D(FVector(PreviewMeshScale, PreviewMeshScale, PreviewMeshScale));
+				PreviewMesh->SetupAttachment(RootComponent);
+				PreviewMesh->RegisterComponent();
+
+				PreviewMeshes_Mount.Add(PreviewMesh);
+			}
+		}
+
+	}
+	check(PreviewMeshes_Mount.Num() == AttachedActors.Num());
+
+	for (int Index = 0; Index < PreviewMeshes_Mount.Num(); ++Index)
+	{
+		const FTransform TargetXform = AttachedActors[Index]->GetActorTransform();
+		PreviewMeshes_Mount[Index]->SetWorldLocation(TargetXform.GetLocation());
+		PreviewMeshes_Mount[Index]->SetVisibility(bShowPreviewMeshes);
+		PreviewMeshes_Mount[Index]->SetWorldScale3D(FVector(PreviewMeshScale, PreviewMeshScale, PreviewMeshScale));
+	}
+
+}
+#endif 
\ No newline at end of file
diff --git a/Engine/Plugins/Experimental/CineCameraRigs/Source/CineCameraRigs/Public/CineCameraAttachMount.h b/Engine/Plugins/Experimental/CineCameraRigs/Source/CineCameraRigs/Public/CineCameraAttachMount.h
new file mode 100644
--- /dev/null
+++ b/Engine/Plugins/Experimental/CineCameraRigs/Source/CineCameraRigs/Public/CineCameraAttachMount.h
@@ -0,0 +1,137 @@
+// Copyright Epic Games, Inc. All Rights Reserved.
+
+#pragma once
+
+#include "CoreMinimal.h"
+#include "UObject/ObjectMacros.h"
+#include "GameFramework/Actor.h"
+#include "Constraint.h"
+
+#include "CineCameraAttachMount.generated.h"
+
+class USpringArmComponent;
+class UBillboardComponent;
+class UTickableParentConstraint;
+
+UCLASS(Blueprintable, Category = "VirtualProduction")
+class ACineCameraAttachMount : public AActor
+{
+	GENERATED_BODY()
+
+public:
+	ACineCameraAttachMount(const FObjectInitializer& ObjectInitializer);
+
+	virtual void Tick(float DeltaTime) override;
+	virtual bool ShouldTickIfViewportsOnly() const override;
+	virtual class USceneComponent* GetDefaultAttachComponent() const override;
+
+	/** Constraint target actor */
+	UPROPERTY(BlueprintReadWrite, EditAnywhere, Category = "Constraint")
+	TSoftObjectPtr<AActor> TargetActor;
+
+	/** Constraint target socket name*/
+	UPROPERTY(BlueprintReadWrite, EditAnywhere, Category = "Constraint")
+	FName TargetSocket;
+
+	/** Constraint axis filter */
+	UPROPERTY(BlueprintReadWrite, EditAnywhere, Category = "Constraint")
+	FTransformFilter TransformFilter;
+
+	/** SprintArm component for lag effect */
+	UPROPERTY(Transient, BlueprintReadOnly, Category = "Default")
+	TObjectPtr<USpringArmComponent> SpringArmComponent;
+
+	/** If true, attached actors lag behind target position to smooth its movement */
+	UPROPERTY(EditAnywhere, BlueprintSetter=SetEnableLocationLag, Category = "Lag")
+	bool bEnableLocationLag = true;
+
+	/** If true, attached actors lag behind target position to smooth its movement */
+	UPROPERTY(EditAnywhere, BlueprintSetter=SetEnableRotationLag, Category = "Lag")
+	bool bEnableRotationLag = true;
+	
+	/** If bEnableLocationLag is true, controls how quickly camera reaches target position. Low values are slower (more lag), high values are faster (less lag), while zero is instant (no lag) */
+	UPROPERTY(EditAnywhere, BlueprintSetter=SetLocationLagSpeed, Category = "Lag", meta = (editcondition = "bEnableLocationLag", ClampMin = "0.0", ClampMax = "1000.0", UIMin = "0.0", UIMax = "1000.0"))
+	float LocationLagSpeed = 10.0f;
+
+	/** If bEnableRotationLag is true, controls how quickly camera reaches target position. Low values are slower (more lag), high values are faster (less lag), while zero is instant (no lag)*/
+	UPROPERTY(EditAnywhere, BlueprintSetter=SetRotationLagSpeed, Category = "Lag", meta = (editcondition = "bEnableRotationLag", ClampMin = "0.0", ClampMax = "1000.0", UIMin = "0.0", UIMax = "1000.0"))
+	float RotationLagSpeed = 10.0f;
+
+	/** Set bEnableLocationLag */
+	UFUNCTION(BlueprintSetter)
+	void SetEnableLocationLag(bool bEnabled);
+
+	/** Set bEnableRotationLag */
+	UFUNCTION(BlueprintSetter)
+	void SetEnableRotationLag(bool bEnabled);
+
+	/** Set LocationLagSpeed */
+	UFUNCTION(BlueprintSetter)
+	void SetLocationLagSpeed(float Speed);
+
+	/** Set RotationLagSpeed*/
+	UFUNCTION(BlueprintSetter)
+	void SetRotationLagSpeed(float Speed);
+
+
+#if WITH_EDITORONLY_DATA
+	/** If enabled, it shows the preview meshes */
+	UPROPERTY(Transient, EditAnywhere, Category = "Preview")
+	bool bShowPreviewMeshes;
+
+	/** Scale of the preview meshes */
+	UPROPERTY(Transient, EditAnywhere, Category = "Preview")
+	float PreviewMeshScale;
+
+	/** Preview mesh for this actor itself*/
+	UPROPERTY(Transient)
+	TObjectPtr<UStaticMeshComponent> PreviewMesh_Root;
+
+	/** Preview mesh*/
+	UPROPERTY(Transient)
+	TObjectPtr<UStaticMesh> MountMesh;
+
+	/** Preview mount mesh components*/
+	UPROPERTY(Transient)
+	TArray<TObjectPtr<UStaticMeshComponent>> PreviewMeshes_Mount;
+
+#endif 
+
+#if WITH_EDITOR
+	virtual void PostEditChangeProperty(FPropertyChangedEvent& PropertyChangedEvent) override;
+	virtual void PostEditChangeChainProperty(FPropertyChangedChainEvent& PropertyChangedEvent) override;
+#endif
+
+	/** Get ParentConstraint object. Returns nullptr if there is no constraint*/
+	UFUNCTION(BlueprintCallable, Category = "Constraint")
+	UTickableParentConstraint* GetConstraint();
+
+	/** Reset constraint offset for location */
+	UFUNCTION(BlueprintCallable, CallInEditor, Category = "Constraint")
+	void ResetLocationOffset();
+
+	/** Reset constraint offset for rotation */
+	UFUNCTION(BlueprintCallable, CallInEditor, Category = "Constraint")
+	void ResetRotationOffset();
+
+	/** Zero out rotationX in world space. This can help setting the mount parallel to the ground. */
+	UFUNCTION(BlueprintCallable, CallInEditor, Category = "Constraint")
+	void ZeroRoll();
+
+	/** Zero out rotationY in world space. This can help setting the mount parallel to the ground. */
+	UFUNCTION(BlueprintCallable, CallInEditor, Category = "Constraint")
+	void ZeroTilt();
+
+private:
+	bool bConstraintUpdated = false;
+
+	void CreateConstraint();
+	void SetAxisFilter();
+	
+#if WITH_EDITOR
+	void CreatePreviewMesh();
+	void UpdatePreviewMeshes();
+#endif
+
+};
+
