Attempt to fix build error

diff --git a/Engine/Plugins/Experimental/CineCameraRigs/Source/CineCameraRigs/Private/CineCameraAttachMount.cpp b/Engine/Plugins/Experimental/CineCameraRigs/Source/CineCameraRigs/Private/CineCameraAttachMount.cpp
--- a/Engine/Plugins/Experimental/CineCameraRigs/Source/CineCameraRigs/Private/CineCameraAttachMount.cpp
+++ b/Engine/Plugins/Experimental/CineCameraRigs/Source/CineCameraRigs/Private/CineCameraAttachMount.cpp
@@ -96,6 +96,7 @@ void ACineCameraAttachMount::PostEditChangeProperty(FPropertyChangedEvent& Prope
 	}
 }
 
+#if WITH_EDITOR
 void ACineCameraAttachMount::PostEditChangeChainProperty(FPropertyChangedChainEvent& PropertyChangedEvent)
 {
 	Super::PostEditChangeChainProperty(PropertyChangedEvent);
@@ -105,6 +106,7 @@ void ACineCameraAttachMount::PostEditChangeChainProperty(FPropertyChangedChainEv
 		SetAxisFilter();
 	}
 }
+#endif
 
 void ACineCameraAttachMount::SetAxisFilter()
 {
