VCam: bPlay property is added in CineCameraRigRail and used by SplineModifer

#virtualized

diff --git a/Engine/Plugins/Experimental/CineCameraRigs/Source/CineCameraRigs/Private/CineCameraRigRail.cpp b/Engine/Plugins/Experimental/CineCameraRigs/Source/CineCameraRigs/Private/CineCameraRigRail.cpp
--- a/Engine/Plugins/Experimental/CineCameraRigs/Source/CineCameraRigs/Private/CineCameraRigRail.cpp
+++ b/Engine/Plugins/Experimental/CineCameraRigs/Source/CineCameraRigs/Private/CineCameraRigRail.cpp
@@ -354,11 +354,11 @@ void ACineCameraRigRail::SetMIDParameters()
 void ACineCameraRigRail::Tick(float DeltaTime)
 {
 	Super::Tick(DeltaTime);
-	if (DriveMode == ECineCameraRigRailDriveMode::Duration)
+	if (DriveMode == ECineCameraRigRailDriveMode::Duration && bPlay)
 	{
 		DriveByParam(DeltaTime);
 	}
-	else if (DriveMode == ECineCameraRigRailDriveMode::Speed)
+	else if (DriveMode == ECineCameraRigRailDriveMode::Speed && bPlay)
 	{
 		DriveBySpeed(DeltaTime);
 	}
diff --git a/Engine/Plugins/Experimental/CineCameraRigs/Source/CineCameraRigs/Public/CineCameraRigRail.h b/Engine/Plugins/Experimental/CineCameraRigs/Source/CineCameraRigs/Public/CineCameraRigRail.h
--- a/Engine/Plugins/Experimental/CineCameraRigs/Source/CineCameraRigs/Public/CineCameraRigRail.h
+++ b/Engine/Plugins/Experimental/CineCameraRigs/Source/CineCameraRigs/Public/CineCameraRigRail.h
@@ -117,6 +117,10 @@ public:
 	UPROPERTY(EditAnywhere, Interp, Category = "DriveMode")
 	float Speed = 100;
 
+	/* Determine if it can update position in Duration mode or Speed mode. If false, it pauses the update.*/
+	UPROPERTY(EditAnywhere, BlueprintReadWrite, Category = "DriveMode")
+	bool bPlay = true;
+
 	/* Enable loop in speed or duration mode */
 	UPROPERTY(EditAnywhere, BlueprintReadWrite, Category = "DriveMode")
 	bool bLoop = true;
diff --git a/Engine/Plugins/Experimental/CineCameraRigs/Source/CineCameraRigsEditor/Private/CineCameraRigRailDetails.cpp b/Engine/Plugins/Experimental/CineCameraRigs/Source/CineCameraRigsEditor/Private/CineCameraRigRailDetails.cpp
--- a/Engine/Plugins/Experimental/CineCameraRigs/Source/CineCameraRigsEditor/Private/CineCameraRigRailDetails.cpp
+++ b/Engine/Plugins/Experimental/CineCameraRigs/Source/CineCameraRigsEditor/Private/CineCameraRigRailDetails.cpp
@@ -219,6 +219,7 @@ void FCineCameraRigRailDetails::CustomizeDriveModeCategory(IDetailLayoutBuilder&
 	DriveModeHandle = DetailBuilder.GetProperty(GET_MEMBER_NAME_CHECKED(ACineCameraRigRail, DriveMode));
 	TSharedPtr<IPropertyHandle> DurationHandle = DetailBuilder.GetProperty(GET_MEMBER_NAME_CHECKED(USplineComponent, Duration), USplineComponent::StaticClass());
 	TSharedPtr<IPropertyHandle> SpeedHandle = DetailBuilder.GetProperty(GET_MEMBER_NAME_CHECKED(ACineCameraRigRail, Speed));
+	TSharedPtr<IPropertyHandle> PlayHandle = DetailBuilder.GetProperty(GET_MEMBER_NAME_CHECKED(ACineCameraRigRail, bPlay));
 	TSharedPtr<IPropertyHandle> LoopHandle = DetailBuilder.GetProperty(GET_MEMBER_NAME_CHECKED(ACineCameraRigRail, bLoop));
 
 	DriveModeCategory.AddProperty(DriveModeHandle.ToSharedRef());
@@ -227,6 +228,8 @@ void FCineCameraRigRailDetails::CustomizeDriveModeCategory(IDetailLayoutBuilder&
 		.EditCondition(TAttribute<bool>::Create([this]() { return GetDriveMode() == ECineCameraRigRailDriveMode::Duration; }), nullptr);
 	DriveModeCategory.AddProperty(SpeedHandle.ToSharedRef())
 		.EditCondition(TAttribute<bool>::Create([this]() { return GetDriveMode() == ECineCameraRigRailDriveMode::Speed; }), nullptr);
+	DriveModeCategory.AddProperty(PlayHandle.ToSharedRef())
+		.EditCondition(TAttribute<bool>::Create([this]() { return GetDriveMode() != ECineCameraRigRailDriveMode::Manual; }), nullptr);
 	DriveModeCategory.AddProperty(LoopHandle.ToSharedRef())
 		.EditCondition(TAttribute<bool>::Create([this]() { return GetDriveMode() != ECineCameraRigRailDriveMode::Manual; }), nullptr);
 
