VCam: ParentModifier to support CineCameraActors for inheritance filter. ACineCameraAttachMount::GetConstraint to be pure function

#virtualized

diff --git a/Engine/Plugins/Experimental/CineCameraRigs/Source/CineCameraRigs/Public/CineCameraAttachMount.h b/Engine/Plugins/Experimental/CineCameraRigs/Source/CineCameraRigs/Public/CineCameraAttachMount.h
--- a/Engine/Plugins/Experimental/CineCameraRigs/Source/CineCameraRigs/Public/CineCameraAttachMount.h
+++ b/Engine/Plugins/Experimental/CineCameraRigs/Source/CineCameraRigs/Public/CineCameraAttachMount.h
@@ -103,7 +103,7 @@ public:
 #endif
 
 	/** Get ParentConstraint object. Returns nullptr if there is no constraint*/
-	UFUNCTION(BlueprintCallable, Category = "Constraint")
+	UFUNCTION(BlueprintPure, Category = "Constraint")
 	UTickableParentConstraint* GetConstraint();
 
 	/** Reset constraint offset for location */
