VCam: Added drive in reverse for Spline. Use Dial for drive mode

#virtualized

diff --git a/Engine/Plugins/Experimental/CineCameraRigs/Source/CineCameraRigs/Private/CineCameraRigRail.cpp b/Engine/Plugins/Experimental/CineCameraRigs/Source/CineCameraRigs/Private/CineCameraRigRail.cpp
--- a/Engine/Plugins/Experimental/CineCameraRigs/Source/CineCameraRigs/Private/CineCameraRigRail.cpp
+++ b/Engine/Plugins/Experimental/CineCameraRigs/Source/CineCameraRigs/Private/CineCameraRigRail.cpp
@@ -377,11 +377,12 @@ void ACineCameraRigRail::PostLoad()
 void ACineCameraRigRail::DriveByParam(float DeltaTime)
 {
 	const float PositionDuration = LastPositionValue() - StartPositionValue();
-	float Param = PositionDuration * UKismetMathLibrary::SafeDivide(DeltaTime, CineSplineComponent->Duration);
+	float Increment = bReverse ? -DeltaTime : DeltaTime;
+	float Param = PositionDuration * UKismetMathLibrary::SafeDivide(Increment, CineSplineComponent->Duration);
 	Param += bUseAbsolutePosition ? AbsolutePositionOnRail : CurrentPositionOnRail;
 	if (bLoop)
 	{
-		Param = PositionDuration > 0.0f ? FMath::Fmod(Param - StartPositionValue(), PositionDuration) : 0.0;
+		Param = PositionDuration > 0.0f ? FMath::Fmod(Param - StartPositionValue() + PositionDuration, PositionDuration) : 0.0;
 		Param += StartPositionValue();
 	}
 	else
@@ -400,7 +401,8 @@ void ACineCameraRigRail::DriveBySpeed(float DeltaTime)
 
 	float TotalTime = CineSplineComponent->GetSplineLength() / FMath::Abs(Speed);
 	float CurrentTime = TotalTime * SpeedProgress;
-	CurrentTime += (FMath::Sign(Speed) * DeltaTime);
+	float Increment = bReverse ? -DeltaTime : DeltaTime;
+	CurrentTime += (FMath::Sign(Speed) * Increment);
 	CurrentTime = bLoop ? FMath::Fmod(CurrentTime + TotalTime, TotalTime) : FMath::Clamp(CurrentTime, 0.0f, TotalTime);
 
 	SpeedProgress = CurrentTime / TotalTime;
diff --git a/Engine/Plugins/Experimental/CineCameraRigs/Source/CineCameraRigs/Public/CineCameraRigRail.h b/Engine/Plugins/Experimental/CineCameraRigs/Source/CineCameraRigs/Public/CineCameraRigRail.h
--- a/Engine/Plugins/Experimental/CineCameraRigs/Source/CineCameraRigs/Public/CineCameraRigRail.h
+++ b/Engine/Plugins/Experimental/CineCameraRigs/Source/CineCameraRigs/Public/CineCameraRigRail.h
@@ -121,6 +121,10 @@ public:
 	UPROPERTY(EditAnywhere, BlueprintReadWrite, Category = "DriveMode")
 	bool bPlay = true;
 
+	/* Determine if it plays in reverse.*/
+	UPROPERTY(EditAnywhere, BlueprintReadWrite, Category = "DriveMode")
+	bool bReverse = false;
+
 	/* Enable loop in speed or duration mode */
 	UPROPERTY(EditAnywhere, BlueprintReadWrite, Category = "DriveMode")
 	bool bLoop = true;
diff --git a/Engine/Plugins/Experimental/CineCameraRigs/Source/CineCameraRigsEditor/Private/CineCameraRigRailDetails.cpp b/Engine/Plugins/Experimental/CineCameraRigs/Source/CineCameraRigsEditor/Private/CineCameraRigRailDetails.cpp
--- a/Engine/Plugins/Experimental/CineCameraRigs/Source/CineCameraRigsEditor/Private/CineCameraRigRailDetails.cpp
+++ b/Engine/Plugins/Experimental/CineCameraRigs/Source/CineCameraRigsEditor/Private/CineCameraRigRailDetails.cpp
@@ -221,6 +221,7 @@ void FCineCameraRigRailDetails::CustomizeDriveModeCategory(IDetailLayoutBuilder&
 	TSharedPtr<IPropertyHandle> SpeedHandle = DetailBuilder.GetProperty(GET_MEMBER_NAME_CHECKED(ACineCameraRigRail, Speed));
 	TSharedPtr<IPropertyHandle> PlayHandle = DetailBuilder.GetProperty(GET_MEMBER_NAME_CHECKED(ACineCameraRigRail, bPlay));
 	TSharedPtr<IPropertyHandle> LoopHandle = DetailBuilder.GetProperty(GET_MEMBER_NAME_CHECKED(ACineCameraRigRail, bLoop));
+	TSharedPtr<IPropertyHandle> ReverseHandle = DetailBuilder.GetProperty(GET_MEMBER_NAME_CHECKED(ACineCameraRigRail, bReverse));
 
 	DriveModeCategory.AddProperty(DriveModeHandle.ToSharedRef());
 
@@ -232,6 +233,8 @@ void FCineCameraRigRailDetails::CustomizeDriveModeCategory(IDetailLayoutBuilder&
 		.EditCondition(TAttribute<bool>::Create([this]() { return GetDriveMode() != ECineCameraRigRailDriveMode::Manual; }), nullptr);
 	DriveModeCategory.AddProperty(LoopHandle.ToSharedRef())
 		.EditCondition(TAttribute<bool>::Create([this]() { return GetDriveMode() != ECineCameraRigRailDriveMode::Manual; }), nullptr);
+	DriveModeCategory.AddProperty(ReverseHandle.ToSharedRef())
+		.EditCondition(TAttribute<bool>::Create([this]() { return GetDriveMode() != ECineCameraRigRailDriveMode::Manual; }), nullptr);
 
 }
 
