VCam: Fix inheritance option -- Add a setter function on CineCameraAttachMount so that it updates internal constraint properly

#virtualized

diff --git a/Engine/Plugins/Experimental/CineCameraRigs/Source/CineCameraRigs/Private/CineCameraAttachMount.cpp b/Engine/Plugins/Experimental/CineCameraRigs/Source/CineCameraRigs/Private/CineCameraAttachMount.cpp
--- a/Engine/Plugins/Experimental/CineCameraRigs/Source/CineCameraRigs/Private/CineCameraAttachMount.cpp
+++ b/Engine/Plugins/Experimental/CineCameraRigs/Source/CineCameraRigs/Private/CineCameraAttachMount.cpp
@@ -103,12 +103,12 @@ void ACineCameraAttachMount::PostEditChangeChainProperty(FPropertyChangedChainEv
 	const FName StructName = PropertyChangedEvent.PropertyChain.GetActiveMemberNode()->GetValue()->GetFName();
 	if (StructName == GET_MEMBER_NAME_CHECKED(ACineCameraAttachMount, TransformFilter))
 	{
-		SetAxisFilter();
+		UpdateAxisFilter();
 	}
 }
 #endif
 
-void ACineCameraAttachMount::SetAxisFilter()
+void ACineCameraAttachMount::UpdateAxisFilter()
 {
 	if (UTickableParentConstraint* Constraint = GetConstraint())
 	{
@@ -234,6 +234,12 @@ void ACineCameraAttachMount::SetRotationLagSpeed(float Speed)
 	SpringArmComponent->CameraRotationLagSpeed = RotationLagSpeed = Speed;
 }
 
+void ACineCameraAttachMount::SetTransformFilter(const FTransformFilter& InFilter)
+{
+	TransformFilter = InFilter;
+	UpdateAxisFilter();
+}
+
 #if WITH_EDITOR
 
 void ACineCameraAttachMount::CreatePreviewMesh()
diff --git a/Engine/Plugins/Experimental/CineCameraRigs/Source/CineCameraRigs/Public/CineCameraAttachMount.h b/Engine/Plugins/Experimental/CineCameraRigs/Source/CineCameraRigs/Public/CineCameraAttachMount.h
--- a/Engine/Plugins/Experimental/CineCameraRigs/Source/CineCameraRigs/Public/CineCameraAttachMount.h
+++ b/Engine/Plugins/Experimental/CineCameraRigs/Source/CineCameraRigs/Public/CineCameraAttachMount.h
@@ -34,7 +34,7 @@ public:
 	FName TargetSocket;
 
 	/** Constraint axis filter */
-	UPROPERTY(BlueprintReadWrite, EditAnywhere, Category = "Constraint")
+	UPROPERTY(EditAnywhere, BlueprintSetter=SetTransformFilter, Category = "Constraint")
 	FTransformFilter TransformFilter;
 
 	/** SprintArm component for lag effect */
@@ -73,6 +73,10 @@ public:
 	UFUNCTION(BlueprintSetter)
 	void SetRotationLagSpeed(float Speed);
 
+	/** Set TransformFilter for the internal constraint*/
+	UFUNCTION(BlueprintSetter)
+	void SetTransformFilter(const FTransformFilter& InFilter);
+
 
 #if WITH_EDITORONLY_DATA
 	/** If enabled, it shows the preview meshes */
@@ -126,7 +130,7 @@ private:
 	bool bConstraintUpdated = false;
 
 	void CreateConstraint();
-	void SetAxisFilter();
+	void UpdateAxisFilter();
 	
 #if WITH_EDITOR
 	void CreatePreviewMesh();
