VCam: World time dilation compensation on JoysticksModifier. Added the option on CineCameraRigRail.

#virtualized

diff --git a/Engine/Plugins/Experimental/CineCameraRigs/Source/CineCameraRigs/Private/CineCameraRigRail.cpp b/Engine/Plugins/Experimental/CineCameraRigs/Source/CineCameraRigs/Private/CineCameraRigRail.cpp
--- a/Engine/Plugins/Experimental/CineCameraRigs/Source/CineCameraRigs/Private/CineCameraRigRail.cpp
+++ b/Engine/Plugins/Experimental/CineCameraRigs/Source/CineCameraRigs/Private/CineCameraRigRail.cpp
@@ -12,6 +12,7 @@
 #include "Engine/World.h"
 #include "Kismet/KismetMathLibrary.h"
 #include "Materials/MaterialInstanceDynamic.h"
+#include "Math/UnrealMathUtility.h"
 #include "UObject/Package.h"
 
 #if WITH_EDITOR
@@ -376,8 +377,9 @@ void ACineCameraRigRail::PostLoad()
 
 void ACineCameraRigRail::DriveByParam(float DeltaTime)
 {
+	const float AdjustedDeltaTime = AdjustDeltaTime(DeltaTime);
 	const float PositionDuration = LastPositionValue() - StartPositionValue();
-	float Increment = bReverse ? -DeltaTime : DeltaTime;
+	const float Increment = bReverse ? -AdjustedDeltaTime : AdjustedDeltaTime;
 	float Param = PositionDuration * UKismetMathLibrary::SafeDivide(Increment, CineSplineComponent->Duration);
 	Param += bUseAbsolutePosition ? AbsolutePositionOnRail : CurrentPositionOnRail;
 	if (bLoop)
@@ -398,10 +400,10 @@ void ACineCameraRigRail::DriveBySpeed(float DeltaTime)
 	{
 		return;
 	}
-
-	float TotalTime = CineSplineComponent->GetSplineLength() / FMath::Abs(Speed);
+	const float AdjustedDeltaTime = AdjustDeltaTime(DeltaTime);
+	const float TotalTime = CineSplineComponent->GetSplineLength() / FMath::Abs(Speed);
 	float CurrentTime = TotalTime * SpeedProgress;
-	float Increment = bReverse ? -DeltaTime : DeltaTime;
+	const float Increment = bReverse ? -AdjustedDeltaTime : AdjustedDeltaTime;
 	CurrentTime += (FMath::Sign(Speed) * Increment);
 	CurrentTime = bLoop ? FMath::Fmod(CurrentTime + TotalTime, TotalTime) : FMath::Clamp(CurrentTime, 0.0f, TotalTime);
 
@@ -409,8 +411,8 @@ void ACineCameraRigRail::DriveBySpeed(float DeltaTime)
 	
 	if (bUseAbsolutePosition)
 	{
-		float Distance = CineSplineComponent->GetSplineLength() * SpeedProgress;
-		float InputKey = CineSplineComponent->GetInputKeyValueAtDistanceAlongSpline(Distance);
+		const float Distance = CineSplineComponent->GetSplineLength() * SpeedProgress;
+		const float InputKey = CineSplineComponent->GetInputKeyValueAtDistanceAlongSpline(Distance);
 		AbsolutePositionOnRail = CineSplineComponent->GetPositionAtInputKey(InputKey);
 	}
 	else
@@ -547,6 +549,12 @@ bool ACineCameraRigRail::IsSequencerDriven()
 #endif
 }
 
+float ACineCameraRigRail::AdjustDeltaTime(float InDeltaTime) const
+{
+	const float TimeDilation = FMath::Max(GetActorTimeDilation(), UE_KINDA_SMALL_NUMBER);
+	return bCompensateTimeScale ? InDeltaTime / TimeDilation : InDeltaTime;
+}
+
 #if WITH_EDITOR
 UMovieSceneFloatTrack* ACineCameraRigRail::FindPositionTrack(const UMovieSceneSequence* InSequence)
 {
diff --git a/Engine/Plugins/Experimental/CineCameraRigs/Source/CineCameraRigs/Public/CineCameraRigRail.h b/Engine/Plugins/Experimental/CineCameraRigs/Source/CineCameraRigs/Public/CineCameraRigRail.h
--- a/Engine/Plugins/Experimental/CineCameraRigs/Source/CineCameraRigs/Public/CineCameraRigRail.h
+++ b/Engine/Plugins/Experimental/CineCameraRigs/Source/CineCameraRigs/Public/CineCameraRigRail.h
@@ -129,6 +129,10 @@ public:
 	UPROPERTY(EditAnywhere, BlueprintReadWrite, Category = "DriveMode")
 	bool bLoop = true;
 
+	/* If enabled, it compensates world time dilation in Speed/Duration mode so that the spline moves as intended speed regardless of recording time scale */
+	UPROPERTY(EditAnywhere, BlueprintReadWrite, Category = "DriveMode")
+	bool bCompensateTimeScale = false;
+
 	/* Set spline mesh material*/
 	UFUNCTION(BlueprintSetter)
 	void SetSplineMeshMaterial(UMaterialInterface* InMaterial);
@@ -175,7 +179,7 @@ protected:
 	float LastPositionValue() const;
 	float SpeedProgress = 0.0f;
 	void UpdateSpeedProgress();
-
+	float AdjustDeltaTime(float InDeltaTime) const;
 
 #if WITH_EDITOR
 	UMovieSceneFloatTrack* FindPositionTrack(const UMovieSceneSequence* InSequence);
