Add setter function for AbsolutePositionOnRail so that it can update the current position via BP

diff --git a/Engine/Plugins/Experimental/CineCameraRigs/Source/CineCameraRigs/Private/CineCameraRigRail.cpp b/Engine/Plugins/Experimental/CineCameraRigs/Source/CineCameraRigs/Private/CineCameraRigRail.cpp
--- a/Engine/Plugins/Experimental/CineCameraRigs/Source/CineCameraRigs/Private/CineCameraRigRail.cpp
+++ b/Engine/Plugins/Experimental/CineCameraRigs/Source/CineCameraRigs/Private/CineCameraRigRail.cpp
@@ -297,6 +297,12 @@ void ACineCameraRigRail::SetSplineMeshTexture(UTexture2D* InTexture)
 	SetMIDParameters();
 }
 
+void ACineCameraRigRail::SetAbsolutePositionOnRail(float Value)
+{
+	AbsolutePositionOnRail = Value;
+	UpdateSpeedProgress();
+}
+
 FVector ACineCameraRigRail::GetVelocityAtPosition(const float InPosition, const float delta) const
 {
 	if (CineSplineComponent->Duration == 0.0f)
diff --git a/Engine/Plugins/Experimental/CineCameraRigs/Source/CineCameraRigs/Public/CineCameraRigRail.h b/Engine/Plugins/Experimental/CineCameraRigs/Source/CineCameraRigs/Public/CineCameraRigRail.h
--- a/Engine/Plugins/Experimental/CineCameraRigs/Source/CineCameraRigs/Public/CineCameraRigRail.h
+++ b/Engine/Plugins/Experimental/CineCameraRigs/Source/CineCameraRigs/Public/CineCameraRigRail.h
@@ -46,7 +46,7 @@ public:
 	bool bUseAbsolutePosition = true;
 
 	/* Custom parameter to drive current position*/
-	UPROPERTY(EditAnywhere, BlueprintReadWrite, Interp, Category = "Rail Controls", meta=(EditCondition="bUseAbsolutePosition"))
+	UPROPERTY(EditAnywhere, Interp, BlueprintSetter="SetAbsolutePositionOnRail", Category = "Rail Controls", meta = (EditCondition = "bUseAbsolutePosition"))
 	float AbsolutePositionOnRail = 1.0f;
 
 	/* Use PointRotation metadata for attachment orientation. If false, attachment orientation is based on the spline curvature*/
@@ -157,6 +157,10 @@ public:
 	UFUNCTION()
 	bool IsSequencerDriven();
 
+	/* Set AbsolutePositionOnRail value*/
+	UFUNCTION(BlueprintSetter)
+	void SetAbsolutePositionOnRail(float Value);
+
 #if WITH_EDITOR
 	virtual void PostEditChangeProperty(FPropertyChangedEvent& PropertyChangedEvent) override;
 #endif
