Fix rig rail animated in sequencer. Fix rig rail crashing when recorded. Fix attachment options when UseAbsolutePosition is off

diff --git a/Engine/Plugins/Experimental/CineCameraRigs/Source/CineCameraRigs/Private/CineCameraRigRail.cpp b/Engine/Plugins/Experimental/CineCameraRigs/Source/CineCameraRigs/Private/CineCameraRigRail.cpp
--- a/Engine/Plugins/Experimental/CineCameraRigs/Source/CineCameraRigs/Private/CineCameraRigRail.cpp
+++ b/Engine/Plugins/Experimental/CineCameraRigs/Source/CineCameraRigs/Private/CineCameraRigRail.cpp
@@ -84,71 +84,65 @@ void ACineCameraRigRail::UpdateRailComponents()
 		return;
 	}
 
-	if (!bUseAbsolutePosition)
-	{
-		Super::UpdateRailComponents();
-	}
-	else
-	{
-		USceneComponent* AttachComponent = GetDefaultAttachComponent();
-		if (CineSplineComponent && AttachComponent)
-		{
-			float InputKey = CineSplineComponent->GetInputKeyAtPosition(AbsolutePositionOnRail);
-			FVector const SplinePosition = CineSplineComponent->GetLocationAtSplineInputKey(InputKey, ESplineCoordinateSpace::World);
-			FQuat SplineQuat = CineSplineComponent->GetQuaternionAtSplineInputKey(InputKey, ESplineCoordinateSpace::World);
 
-			if (bUsePointRotation)
-			{
-				SplineQuat = CineSplineComponent->GetComponentTransform().GetRotation() * CineSplineComponent->GetPointRotationAtSplineInputKey(InputKey);
-			}
-			FVector Position = GetActorLocation();
-			FRotator Rotation = GetActorRotation();
-			if (bAttachLocationX)
-			{
-				Position.X = SplinePosition.X;
-			}
-			if (bAttachLocationY)
-			{
-				Position.Y = SplinePosition.Y;
-			}
-			if (bAttachLocationZ)
-			{
-				Position.Z = SplinePosition.Z;
-			}
-			if (bAttachRotationX)
-			{
-				Rotation.Roll = SplineQuat.Rotator().Roll;
-			}
-			if (bAttachRotationY)
-			{
-				Rotation.Pitch = SplineQuat.Rotator().Pitch;
-			}
-			if (bAttachRotationZ)
-			{
-				Rotation.Yaw = SplineQuat.Rotator().Yaw;
-			}
-			if (bLockOrientationToRail)
-			{
-				AttachComponent->SetWorldTransform(FTransform(Rotation, Position));
-			}
-			else
-			{
-				AttachComponent->SetWorldLocation(SplinePosition);
-			}
-		}
-	}
-	TArray< AActor* > AttachedActors;
-	GetAttachedActors(AttachedActors);
+	FVector SplinePosition;
+	FQuat SplineQuat;
 	float InputKey = 0.0f;
+
 	if (bUseAbsolutePosition)
 	{
 		InputKey = CineSplineComponent->GetInputKeyAtPosition(AbsolutePositionOnRail);
+		SplinePosition = CineSplineComponent->GetLocationAtSplineInputKey(InputKey, ESplineCoordinateSpace::World);
+		SplineQuat = CineSplineComponent->GetQuaternionAtSplineInputKey(InputKey, ESplineCoordinateSpace::World);
 	}
 	else
 	{
-		float const SplineLen = RailSplineComponent->GetSplineLength();
+		float const SplineLen = CineSplineComponent->GetSplineLength();
+		SplinePosition = CineSplineComponent->GetLocationAtDistanceAlongSpline(CurrentPositionOnRail * SplineLen, ESplineCoordinateSpace::World);
+		SplineQuat = CineSplineComponent->GetQuaternionAtDistanceAlongSpline(CurrentPositionOnRail * SplineLen, ESplineCoordinateSpace::World);
 		InputKey = CineSplineComponent->GetInputKeyValueAtDistanceAlongSpline(CurrentPositionOnRail * SplineLen);
 	}
+
+	if (bUsePointRotation)
+	{
+		SplineQuat = CineSplineComponent->GetComponentTransform().GetRotation() * CineSplineComponent->GetPointRotationAtSplineInputKey(InputKey);
+	}
+
+	if (USceneComponent* AttachComponent = GetDefaultAttachComponent())
+	{
+		FVector Position = GetActorLocation();
+		FRotator Rotation = GetActorRotation();
+
+		if (bAttachLocationX)
+		{
+			Position.X = SplinePosition.X;
+		}
+		if (bAttachLocationY)
+		{
+			Position.Y = SplinePosition.Y;
+		}
+		if (bAttachLocationZ)
+		{
+			Position.Z = SplinePosition.Z;
+		}
+		if (bAttachRotationX)
+		{
+			Rotation.Roll = SplineQuat.Rotator().Roll;
+		}
+		if (bAttachRotationY)
+		{
+			Rotation.Pitch = SplineQuat.Rotator().Pitch;
+		}
+		if (bAttachRotationZ)
+		{
+			Rotation.Yaw = SplineQuat.Rotator().Yaw;
+		}
+		AttachComponent->SetWorldTransform(FTransform(Rotation, Position));
+	}
+
+	TArray< AActor* > AttachedActors;
+	GetAttachedActors(AttachedActors);
+
 	float FocalLengthValue = CineSplineComponent->GetFloatPropertyAtSplineInputKey(InputKey, FName("FocalLength"));
 	float ApertureValue = CineSplineComponent->GetFloatPropertyAtSplineInputKey(InputKey, FName("Aperture"));
 	float FocusDistanceValue = CineSplineComponent->GetFloatPropertyAtSplineInputKey(InputKey, FName("FocusDistance"));
@@ -622,6 +616,11 @@ UMovieSceneFloatTrack* ACineCameraRigRail::FindPositionTrack(const UMovieSceneSe
 
 void ACineCameraRigRail::OnSequencerCheck()
 {
+	if (!GetWorld())
+	{
+		return;
+	}
+
 	bool bHasKey = false;
 	ULevelSequence* LevelSequence = ULevelSequenceEditorBlueprintLibrary::GetCurrentLevelSequence();
 	UMovieSceneFloatTrack* Track = FindPositionTrack(LevelSequence);
diff --git a/Engine/Plugins/Experimental/CineCameraRigs/Source/CineCameraRigs/Private/CineCameraRigs.cpp b/Engine/Plugins/Experimental/CineCameraRigs/Source/CineCameraRigs/Private/CineCameraRigs.cpp
--- a/Engine/Plugins/Experimental/CineCameraRigs/Source/CineCameraRigs/Private/CineCameraRigs.cpp
+++ b/Engine/Plugins/Experimental/CineCameraRigs/Source/CineCameraRigs/Private/CineCameraRigs.cpp
@@ -5,6 +5,9 @@
 #include "ConcertSyncSettings.h"
 #include "Algo/AnyOf.h"
 
+#include "CineCameraRigRail.h"
+#include "MovieSceneTracksComponentTypes.h"
+
 #define LOCTEXT_NAMESPACE "FCineCameraRigsModule"
 
 DEFINE_LOG_CATEGORY(LogCineSpline);
@@ -32,6 +35,22 @@ void FCineCameraRigsModule::StartupModule()
 		}
 
 	}
+
+	using namespace UE::MovieScene;
+	FMovieSceneTracksComponentTypes* TracksComponents = FMovieSceneTracksComponentTypes::Get();
+	TracksComponents->Accessors.Float.Add(
+		ACineCameraRigRail::StaticClass(), GET_MEMBER_NAME_CHECKED(ACineCameraRigRail, AbsolutePositionOnRail),
+		GetAbsolutePositionOnRail, SetAbsolutePositionOnRail);
+}
+
+float FCineCameraRigsModule::GetAbsolutePositionOnRail(const UObject* Object)
+{
+	return CastChecked<const ACineCameraRigRail>(Object)->AbsolutePositionOnRail;
+}
+
+void  FCineCameraRigsModule::SetAbsolutePositionOnRail(UObject* Object, float InNewValue)
+{
+	CastChecked<ACineCameraRigRail>(Object)->SetAbsolutePositionOnRail(InNewValue);
 }
 
 void FCineCameraRigsModule::ShutdownModule()
diff --git a/Engine/Plugins/Experimental/CineCameraRigs/Source/CineCameraRigs/Public/CineCameraRigs.h b/Engine/Plugins/Experimental/CineCameraRigs/Source/CineCameraRigs/Public/CineCameraRigs.h
--- a/Engine/Plugins/Experimental/CineCameraRigs/Source/CineCameraRigs/Public/CineCameraRigs.h
+++ b/Engine/Plugins/Experimental/CineCameraRigs/Source/CineCameraRigs/Public/CineCameraRigs.h
@@ -12,4 +12,7 @@ public:
 	/** IModuleInterface implementation */
 	virtual void StartupModule() override;
 	virtual void ShutdownModule() override;
+
+	static float GetAbsolutePositionOnRail(const UObject* Object);
+	static void SetAbsolutePositionOnRail(UObject* Object, float InNewValue);
 };
