CineCameraRigRail: update rotation attachment options if LockOrientationToRail is set false manually

diff --git a/Engine/Plugins/Experimental/CineCameraRigs/Source/CineCameraRigs/Private/CineCameraRigRail.cpp b/Engine/Plugins/Experimental/CineCameraRigs/Source/CineCameraRigs/Private/CineCameraRigRail.cpp
--- a/Engine/Plugins/Experimental/CineCameraRigs/Source/CineCameraRigs/Private/CineCameraRigRail.cpp
+++ b/Engine/Plugins/Experimental/CineCameraRigs/Source/CineCameraRigs/Private/CineCameraRigRail.cpp
@@ -125,15 +125,15 @@ void ACineCameraRigRail::UpdateRailComponents()
 		{
 			Position.Z = SplinePosition.Z;
 		}
-		if (bAttachRotationX)
+		if (bAttachRotationX && bLockOrientationToRail)
 		{
 			Rotation.Roll = SplineQuat.Rotator().Roll;
 		}
-		if (bAttachRotationY)
+		if (bAttachRotationY && bLockOrientationToRail)
 		{
 			Rotation.Pitch = SplineQuat.Rotator().Pitch;
 		}
-		if (bAttachRotationZ)
+		if (bAttachRotationZ && bLockOrientationToRail)
 		{
 			Rotation.Yaw = SplineQuat.Rotator().Yaw;
 		}
diff --git a/Engine/Plugins/Experimental/CineCameraRigs/Source/CineCameraRigs/Public/CineCameraRigRail.h b/Engine/Plugins/Experimental/CineCameraRigs/Source/CineCameraRigs/Public/CineCameraRigRail.h
--- a/Engine/Plugins/Experimental/CineCameraRigs/Source/CineCameraRigs/Public/CineCameraRigRail.h
+++ b/Engine/Plugins/Experimental/CineCameraRigs/Source/CineCameraRigs/Public/CineCameraRigRail.h
@@ -50,7 +50,7 @@ public:
 	float AbsolutePositionOnRail = 1.0f;
 
 	/* Use PointRotation metadata for attachment orientation. If false, attachment orientation is based on the spline curvature*/
-	UPROPERTY(EditAnywhere, BlueprintReadWrite, Category = "Rail Controls", meta = (EditCondition = "bLockOrientationToRail"))
+	UPROPERTY(EditAnywhere, BlueprintReadWrite, Category = "Rail Controls")
 	bool bUsePointRotation = true;
 
 	/* Material assigned to spline component mesh*/
diff --git a/Engine/Plugins/Experimental/CineCameraRigs/Source/CineCameraRigsEditor/Private/CineCameraRigRailDetails.cpp b/Engine/Plugins/Experimental/CineCameraRigs/Source/CineCameraRigsEditor/Private/CineCameraRigRailDetails.cpp
--- a/Engine/Plugins/Experimental/CineCameraRigs/Source/CineCameraRigsEditor/Private/CineCameraRigRailDetails.cpp
+++ b/Engine/Plugins/Experimental/CineCameraRigs/Source/CineCameraRigsEditor/Private/CineCameraRigRailDetails.cpp
@@ -128,6 +128,14 @@ void FCineCameraRigRailDetails::CustomizeSplineVisualizationCategory(IDetailLayo
 
 void FCineCameraRigRailDetails::CustomizeAttachmentCategory(IDetailLayoutBuilder& DetailBuilder)
 {
+	TSharedRef<IPropertyHandle> LockOrientationPropertyHandle = DetailBuilder.GetProperty(GET_MEMBER_NAME_CHECKED(ACameraRig_Rail, bLockOrientationToRail), ACameraRig_Rail::StaticClass());
+	const TAttribute<bool> RotationEditCondition = TAttribute<bool>::Create([this, LockOrientationPropertyHandle]()
+		{
+			bool bCond = false;
+			LockOrientationPropertyHandle->GetValue(bCond);
+			return bCond;
+		});
+
 	IDetailCategoryBuilder& AttachmentCategory = DetailBuilder.EditCategory("Attachment");
 
 	// Attachment : Location
@@ -166,23 +174,67 @@ void FCineCameraRigRailDetails::CustomizeAttachmentCategory(IDetailLayoutBuilder
 
 	IDetailGroup& RotationGroup = AttachmentCategory.AddGroup(TEXT("Rotation"), LOCTEXT("RotationLabel", "Rotation"));
 	RotationGroup.HeaderRow()
+		.EditCondition(RotationEditCondition, nullptr)
 		.NameContent()
 		[
 			SNew(STextBlock)
 			.Text(LOCTEXT("RotationLabel", "Rotation"))
 			.Font(IDetailLayoutBuilder::GetDetailFont())
+			.ToolTipText_Lambda([LockOrientationPropertyHandle]()
+				{
+					bool LockOrientation;
+					LockOrientationPropertyHandle.Get().GetValue(LockOrientation);
+					if (!LockOrientation)
+					{
+						return FText(LOCTEXT("RotationAttachmentToolTip", "Disabled because LockOrientationToRail is false"));
+					}
+					return LOCTEXT("RotationAttachmentToolTip", "Determines if camera mount inherits Rotation");
+				})
 		]
 		.ValueContent()
 		[
 			SNew(SCheckBox)
-			.IsChecked(this, &FCineCameraRigRailDetails::IsAttachOptionChecked, RotationPropertyHandles)
+			.IsChecked(this, &FCineCameraRigRailDetails::IsAttachRotationOptionChecked, RotationPropertyHandles)
 			.OnCheckStateChanged(this, &FCineCameraRigRailDetails::OnAttachOptionChanged, RotationPropertyHandles)
 		];
 
 	for (auto& PropertyHandle : RotationPropertyHandles)
 	{
 		DetailBuilder.HideProperty(PropertyHandle);
-		RotationGroup.AddPropertyRow(PropertyHandle);
+		IDetailPropertyRow& RotationAxisRow = RotationGroup.AddPropertyRow(PropertyHandle);
+		RotationAxisRow.EditCondition(RotationEditCondition, nullptr);
+
+		RotationAxisRow.CustomWidget(false)
+			.NameContent()
+			[
+				SNew(STextBlock)
+					.Text(PropertyHandle->GetPropertyDisplayName())
+					.Font(IDetailLayoutBuilder::GetDetailFont())
+					.ToolTipText_Lambda([PropertyHandle, LockOrientationPropertyHandle]()
+						{
+							bool LockOrientation;
+							LockOrientationPropertyHandle.Get().GetValue(LockOrientation);
+							if (!LockOrientation)
+							{
+								return FText(LOCTEXT("RotationAxisAttachmentToolTip", "Disabled because LockOrientationToRail is false"));
+							}
+							return PropertyHandle->GetToolTipText();
+						})
+			]
+			.ValueContent()
+			.MinDesiredWidth(125.0f)
+			[
+				SNew(SCheckBox)
+					.IsChecked_Lambda([PropertyHandle, LockOrientationPropertyHandle]() -> ECheckBoxState
+						{
+							bool IsChecked;
+							PropertyHandle.Get().GetValue(IsChecked);
+							bool LockOrientation;
+							LockOrientationPropertyHandle.Get().GetValue(LockOrientation);
+							return IsChecked && LockOrientation ? ECheckBoxState::Checked : ECheckBoxState::Unchecked;
+						})
+					.OnCheckStateChanged(this, &FCineCameraRigRailDetails::OnRotationAxisAttachmentChanged, PropertyHandle)
+			];
 	}
 
 	// Attachment : Camera
@@ -291,6 +343,18 @@ ECheckBoxState FCineCameraRigRailDetails::IsAttachOptionChecked(TArray<TSharedRe
 	return ECheckBoxState::Undetermined;
 }
 
+ECheckBoxState FCineCameraRigRailDetails::IsAttachRotationOptionChecked(TArray<TSharedRef<IPropertyHandle>> PropertyHandles) const
+{
+	if (ACineCameraRigRail* RigRailActor = RigRailActorPtr.Get())
+	{
+		if (!RigRailActor->bLockOrientationToRail)
+		{
+			return ECheckBoxState::Unchecked;
+		}
+	}
+	return IsAttachOptionChecked(PropertyHandles);
+}
+
 void FCineCameraRigRailDetails::OnAttachOptionChanged(ECheckBoxState NewState, TArray<TSharedRef<IPropertyHandle>> PropertyHandles)
 {
 	if (NewState == ECheckBoxState::Undetermined)
@@ -423,4 +487,17 @@ TOptional<float> FCineCameraRigRailDetails::GetAbsolutePositionSliderMaxValue()
 	}
 	return MaxValue;
 }
+
+void FCineCameraRigRailDetails::OnRotationAxisAttachmentChanged(ECheckBoxState NewState, TSharedRef<IPropertyHandle> PropertyHandle)
+{
+	if (NewState == ECheckBoxState::Undetermined)
+	{
+		return;
+	}
+	const bool bValue = (NewState == ECheckBoxState::Checked) ? true : false;
+	if (PropertyHandle->IsValidHandle())
+	{
+		PropertyHandle->SetValue(bValue);
+	}
+}
 #undef LOCTEXT_NAMESPACE
\ No newline at end of file
diff --git a/Engine/Plugins/Experimental/CineCameraRigs/Source/CineCameraRigsEditor/Private/CineCameraRigRailDetails.h b/Engine/Plugins/Experimental/CineCameraRigs/Source/CineCameraRigsEditor/Private/CineCameraRigRailDetails.h
--- a/Engine/Plugins/Experimental/CineCameraRigs/Source/CineCameraRigsEditor/Private/CineCameraRigRailDetails.h
+++ b/Engine/Plugins/Experimental/CineCameraRigs/Source/CineCameraRigsEditor/Private/CineCameraRigRailDetails.h
@@ -54,4 +54,7 @@ private:
 	void CustomizeAttachmentCategory(IDetailLayoutBuilder& DetailBuilder);
 	void CustomizeDriveModeCategory(IDetailLayoutBuilder& DetailBuilder);
 	void HideExtraCategories(IDetailLayoutBuilder& DetailBuilder);
+
+	ECheckBoxState IsAttachRotationOptionChecked(TArray<TSharedRef<IPropertyHandle>> PropertyHandles) const;
+	void OnRotationAxisAttachmentChanged(ECheckBoxState NewState, TSharedRef<IPropertyHandle> PropertyHandle);
 };
