Fix conflict LOCTEXT key
#rnx

diff --git a/Engine/Plugins/Experimental/CineCameraRigs/Source/CineCameraRigsEditor/Private/CineCameraRigRailDetails.cpp b/Engine/Plugins/Experimental/CineCameraRigs/Source/CineCameraRigsEditor/Private/CineCameraRigRailDetails.cpp
--- a/Engine/Plugins/Experimental/CineCameraRigs/Source/CineCameraRigsEditor/Private/CineCameraRigRailDetails.cpp
+++ b/Engine/Plugins/Experimental/CineCameraRigs/Source/CineCameraRigsEditor/Private/CineCameraRigRailDetails.cpp
@@ -186,7 +186,7 @@ void FCineCameraRigRailDetails::CustomizeAttachmentCategory(IDetailLayoutBuilder
 					LockOrientationPropertyHandle.Get().GetValue(LockOrientation);
 					if (!LockOrientation)
 					{
-						return FText(LOCTEXT("RotationAttachmentToolTip", "Disabled because LockOrientationToRail is false"));
+						return FText(LOCTEXT("RotationAttachmentDisabledToolTip", "Disabled because LockOrientationToRail is false"));
 					}
 					return LOCTEXT("RotationAttachmentToolTip", "Determines if camera mount inherits Rotation");
 				})
