CineCameraRigRail: Fix for rail mesh not updating when UseAbsolutePosition is off

diff --git a/Engine/Plugins/Experimental/CineCameraRigs/Source/CineCameraRigs/Private/CineCameraRigRail.cpp b/Engine/Plugins/Experimental/CineCameraRigs/Source/CineCameraRigs/Private/CineCameraRigRail.cpp
--- a/Engine/Plugins/Experimental/CineCameraRigs/Source/CineCameraRigs/Private/CineCameraRigRail.cpp
+++ b/Engine/Plugins/Experimental/CineCameraRigs/Source/CineCameraRigs/Private/CineCameraRigRail.cpp
@@ -170,11 +170,7 @@ void ACineCameraRigRail::UpdateRailComponents()
 		const UWorld* const MyWorld = GetWorld();
 		if (MyWorld && !MyWorld->IsGameWorld())
 		{
-			// If bUseAbsolutePosition is false, UpdatePreviewMeshes() is already called inside Super::UpdateRailComponents()
-			if (bUseAbsolutePosition)
-			{
-				UpdatePreviewMeshes();
-			}
+			UpdatePreviewMeshes();
 		}
 
 		// Set HiddenInGame false on the spline mesh so that it can show up in game mode
