Fix spline point info drawn behind the camera

diff --git a/Engine/Plugins/Experimental/CineCameraRigs/Source/CineCameraRigsEditor/Private/CineSplineComponentVisualizer.cpp b/Engine/Plugins/Experimental/CineCameraRigs/Source/CineCameraRigsEditor/Private/CineSplineComponentVisualizer.cpp
--- a/Engine/Plugins/Experimental/CineCameraRigs/Source/CineCameraRigsEditor/Private/CineSplineComponentVisualizer.cpp
+++ b/Engine/Plugins/Experimental/CineCameraRigs/Source/CineCameraRigsEditor/Private/CineSplineComponentVisualizer.cpp
@@ -112,6 +112,10 @@ void FCineSplineComponentVisualizer::DrawVisualizationHUD(const UActorComponent*
 				FVector Location = SplineComp->GetLocationAtSplinePoint(i, ESplineCoordinateSpace::World);
 				FPlane Plane(0, 0, 0, 0);
 				Plane = View->Project(Location);
+				if (Plane.Z < 0.0)
+				{
+					continue;
+				}
 				const FVector Position(Plane);
 				const float DrawPositionX = FMath::FloorToFloat(HalfX + Position.X * HalfX);
 				const float DrawPositionY = FMath::FloorToFloat(HalfY + -1.f * Position.Y * HalfY);
