Bounce option in CineCameraRigRail

#virtualized

diff --git a/Engine/Plugins/Experimental/CineCameraRigs/Source/CineCameraRigs/Private/CineCameraRigRail.cpp b/Engine/Plugins/Experimental/CineCameraRigs/Source/CineCameraRigs/Private/CineCameraRigRail.cpp
--- a/Engine/Plugins/Experimental/CineCameraRigs/Source/CineCameraRigs/Private/CineCameraRigRail.cpp
+++ b/Engine/Plugins/Experimental/CineCameraRigs/Source/CineCameraRigs/Private/CineCameraRigRail.cpp
@@ -375,10 +375,21 @@ void ACineCameraRigRail::DriveByParam(float DeltaTime)
 {
 	const float AdjustedDeltaTime = AdjustDeltaTime(DeltaTime);
 	const float PositionDuration = LastPositionValue() - StartPositionValue();
-	const float Increment = bReverse ? -AdjustedDeltaTime : AdjustedDeltaTime;
+	const float Increment = PlayForward() ? AdjustedDeltaTime : -AdjustedDeltaTime;
 	float Param = PositionDuration * UKismetMathLibrary::SafeDivide(Increment, CineSplineComponent->Duration);
 	Param += bUseAbsolutePosition ? AbsolutePositionOnRail : CurrentPositionOnRail;
-	if (bLoop)
+
+	// Reverse direction in Bounce mode when Param goes outside of the range 
+	if (LoopMode == ECineCameraRigRailLoopMode::Bounce)
+	{
+		if (Param < StartPositionValue() || Param > LastPositionValue())
+		{
+			bBounceForward = !bBounceForward;
+		}
+	}
+
+	// Mod or Clamp the Param variable so that it's within the range
+	if (LoopMode == ECineCameraRigRailLoopMode::Loop)
 	{
 		Param = PositionDuration > 0.0f ? FMath::Fmod(Param - StartPositionValue() + PositionDuration, PositionDuration) : 0.0;
 		Param += StartPositionValue();
@@ -399,9 +410,20 @@ void ACineCameraRigRail::DriveBySpeed(float DeltaTime)
 	const float AdjustedDeltaTime = AdjustDeltaTime(DeltaTime);
 	const float TotalTime = CineSplineComponent->GetSplineLength() / FMath::Abs(Speed);
 	float CurrentTime = TotalTime * SpeedProgress;
-	const float Increment = bReverse ? -AdjustedDeltaTime : AdjustedDeltaTime;
+	const float Increment = PlayForward() ? AdjustedDeltaTime : -AdjustedDeltaTime;
 	CurrentTime += (FMath::Sign(Speed) * Increment);
-	CurrentTime = bLoop ? FMath::Fmod(CurrentTime + TotalTime, TotalTime) : FMath::Clamp(CurrentTime, 0.0f, TotalTime);
+
+	// Reverse direction in Bounce mode when CurrenTime goes outside of the range
+	if (LoopMode == ECineCameraRigRailLoopMode::Bounce)
+	{
+		if (CurrentTime < 0.0f || CurrentTime > TotalTime)
+		{
+			bBounceForward = !bBounceForward;
+		}
+	}
+
+	// Mod or Clamp the CurrentTime variable so that it's within the range
+	CurrentTime = LoopMode == ECineCameraRigRailLoopMode::Loop ? FMath::Fmod(CurrentTime + TotalTime, TotalTime) : FMath::Clamp(CurrentTime, 0.0f, TotalTime);
 
 	SpeedProgress = CurrentTime / TotalTime;
 	
@@ -551,6 +573,15 @@ float ACineCameraRigRail::AdjustDeltaTime(float InDeltaTime) const
 	return bCompensateTimeScale ? InDeltaTime / TimeDilation : InDeltaTime;
 }
 
+bool ACineCameraRigRail::PlayForward() const
+{
+	if (LoopMode == ECineCameraRigRailLoopMode::Bounce)
+	{
+		return bBounceForward;
+	}
+	return !bReverse;
+}
+
 #if WITH_EDITOR
 UMovieSceneFloatTrack* ACineCameraRigRail::FindPositionTrack(const UMovieSceneSequence* InSequence)
 {
diff --git a/Engine/Plugins/Experimental/CineCameraRigs/Source/CineCameraRigs/Public/CineCameraRigRail.h b/Engine/Plugins/Experimental/CineCameraRigs/Source/CineCameraRigs/Public/CineCameraRigRail.h
--- a/Engine/Plugins/Experimental/CineCameraRigs/Source/CineCameraRigs/Public/CineCameraRigRail.h
+++ b/Engine/Plugins/Experimental/CineCameraRigs/Source/CineCameraRigs/Public/CineCameraRigRail.h
@@ -26,6 +26,19 @@ enum class ECineCameraRigRailDriveMode : uint8
 	Speed,
 };
 
+UENUM(BlueprintType)
+enum class ECineCameraRigRailLoopMode : uint8
+{
+	/** Play once*/
+	NoLoop,
+
+	/** Play Loop*/
+	Loop,
+
+	/** Play forward then backward*/
+	Bounce,
+};
+
 UCLASS(Blueprintable, Category = "VirtualProduction")
 class CINECAMERARIGS_API ACineCameraRigRail : public ACameraRig_Rail
 {
@@ -122,13 +135,16 @@ public:
 	bool bPlay = true;
 
 	/* Determine if it plays in reverse.*/
-	UPROPERTY(EditAnywhere, BlueprintReadWrite, Category = "DriveMode")
+	UPROPERTY(EditAnywhere, BlueprintReadWrite, Category = "DriveMode", meta = (EditCondition="LoopMode != ECineCameraRigRailLoopMode::Bounce"))
 	bool bReverse = false;
 
 	/* Enable loop in speed or duration mode */
-	UPROPERTY(EditAnywhere, BlueprintReadWrite, Category = "DriveMode")
+	UPROPERTY(BlueprintReadWrite, Category = "DriveMode", meta = (DeprecatedProperty, DeprecationMessage = "Use LoopMode instead"))
 	bool bLoop = true;
 
+	UPROPERTY(EditAnywhere, BlueprintReadWrite, Category = "DriveMode")
+	ECineCameraRigRailLoopMode LoopMode = ECineCameraRigRailLoopMode::Loop;
+
 	/* If enabled, it compensates world time dilation in Speed/Duration mode so that the spline moves as intended speed regardless of recording time scale */
 	UPROPERTY(EditAnywhere, BlueprintReadWrite, Category = "DriveMode")
 	bool bCompensateTimeScale = false;
@@ -185,6 +201,9 @@ protected:
 	void UpdateSpeedProgress();
 	float AdjustDeltaTime(float InDeltaTime) const;
 
+	bool bBounceForward = true;
+	bool PlayForward() const;
+
 #if WITH_EDITOR
 	UMovieSceneFloatTrack* FindPositionTrack(const UMovieSceneSequence* InSequence);
 
