CineCameraRigRail is hidden in game by default
#rnx

diff --git a/Engine/Plugins/Experimental/CineCameraRigs/Source/CineCameraRigs/Private/CineCameraRigRail.cpp b/Engine/Plugins/Experimental/CineCameraRigs/Source/CineCameraRigs/Private/CineCameraRigRail.cpp
--- a/Engine/Plugins/Experimental/CineCameraRigs/Source/CineCameraRigs/Private/CineCameraRigRail.cpp
+++ b/Engine/Plugins/Experimental/CineCameraRigs/Source/CineCameraRigs/Private/CineCameraRigRail.cpp
@@ -56,6 +56,7 @@ ACineCameraRigRail::ACineCameraRigRail(const FObjectInitializer& ObjectInitializ
 	}
 
 	bLockOrientationToRail = true;
+	SetActorHiddenInGame(true);
 
 #if WITH_EDITORONLY_DATA
 	if (PreviewMesh_Mount)
