CineCameraRigs: UE_API

diff --git a/Engine/Plugins/Experimental/CineCameraRigs/Source/CineCameraRigs/Public/CineCameraAttachMount.h b/Engine/Plugins/Experimental/CineCameraRigs/Source/CineCameraRigs/Public/CineCameraAttachMount.h
--- a/Engine/Plugins/Experimental/CineCameraRigs/Source/CineCameraRigs/Public/CineCameraAttachMount.h
+++ b/Engine/Plugins/Experimental/CineCameraRigs/Source/CineCameraRigs/Public/CineCameraAttachMount.h
@@ -9,6 +9,8 @@
 
 #include "CineCameraAttachMount.generated.h"
 
+#define UE_API CINECAMERARIGS_API
+
 class USpringArmComponent;
 class UBillboardComponent;
 class UTickableParentConstraint;
@@ -19,11 +21,11 @@ class ACineCameraAttachMount : public AActor
 	GENERATED_BODY()
 
 public:
-	ACineCameraAttachMount(const FObjectInitializer& ObjectInitializer);
+	UE_API ACineCameraAttachMount(const FObjectInitializer& ObjectInitializer);
 
-	virtual void Tick(float DeltaTime) override;
-	virtual bool ShouldTickIfViewportsOnly() const override;
-	virtual class USceneComponent* GetDefaultAttachComponent() const override;
+	UE_API virtual void Tick(float DeltaTime) override;
+	UE_API virtual bool ShouldTickIfViewportsOnly() const override;
+	UE_API virtual class USceneComponent* GetDefaultAttachComponent() const override;
 
 	/** Constraint target actor */
 	UPROPERTY(BlueprintReadWrite, EditAnywhere, Category = "Constraint")
@@ -59,23 +61,23 @@ public:
 
 	/** Set bEnableLocationLag */
 	UFUNCTION(BlueprintSetter)
-	void SetEnableLocationLag(bool bEnabled);
+	UE_API void SetEnableLocationLag(bool bEnabled);
 
 	/** Set bEnableRotationLag */
 	UFUNCTION(BlueprintSetter)
-	void SetEnableRotationLag(bool bEnabled);
+	UE_API void SetEnableRotationLag(bool bEnabled);
 
 	/** Set LocationLagSpeed */
 	UFUNCTION(BlueprintSetter)
-	void SetLocationLagSpeed(float Speed);
+	UE_API void SetLocationLagSpeed(float Speed);
 
 	/** Set RotationLagSpeed*/
 	UFUNCTION(BlueprintSetter)
-	void SetRotationLagSpeed(float Speed);
+	UE_API void SetRotationLagSpeed(float Speed);
 
 	/** Set TransformFilter for the internal constraint*/
 	UFUNCTION(BlueprintSetter)
-	void SetTransformFilter(const FTransformFilter& InFilter);
+	UE_API void SetTransformFilter(const FTransformFilter& InFilter);
 
 
 #if WITH_EDITORONLY_DATA
@@ -102,29 +104,29 @@ public:
 #endif 
 
 #if WITH_EDITOR
-	virtual void PostEditChangeProperty(FPropertyChangedEvent& PropertyChangedEvent) override;
-	virtual void PostEditChangeChainProperty(FPropertyChangedChainEvent& PropertyChangedEvent) override;
+	UE_API virtual void PostEditChangeProperty(FPropertyChangedEvent& PropertyChangedEvent) override;
+	UE_API virtual void PostEditChangeChainProperty(FPropertyChangedChainEvent& PropertyChangedEvent) override;
 #endif
 
 	/** Get ParentConstraint object. Returns nullptr if there is no constraint*/
 	UFUNCTION(BlueprintPure, Category = "Constraint")
-	UTickableParentConstraint* GetConstraint();
+	UE_API UTickableParentConstraint* GetConstraint();
 
 	/** Reset constraint offset for location */
 	UFUNCTION(BlueprintCallable, CallInEditor, Category = "Constraint")
-	void ResetLocationOffset();
+	UE_API void ResetLocationOffset();
 
 	/** Reset constraint offset for rotation */
 	UFUNCTION(BlueprintCallable, CallInEditor, Category = "Constraint")
-	void ResetRotationOffset();
+	UE_API void ResetRotationOffset();
 
 	/** Zero out rotationX in world space. This can help setting the mount parallel to the ground. */
 	UFUNCTION(BlueprintCallable, CallInEditor, Category = "Constraint")
-	void ZeroRoll();
+	UE_API void ZeroRoll();
 
 	/** Zero out rotationY in world space. This can help setting the mount parallel to the ground. */
 	UFUNCTION(BlueprintCallable, CallInEditor, Category = "Constraint")
-	void ZeroTilt();
+	UE_API void ZeroTilt();
 
 private:
 	bool bConstraintUpdated = false;
@@ -139,3 +141,4 @@ private:
 
 };
 
+#undef UE_API
\ No newline at end of file
diff --git a/Engine/Plugins/Experimental/CineCameraRigs/Source/CineCameraRigs/Public/CineCameraRigRail.h b/Engine/Plugins/Experimental/CineCameraRigs/Source/CineCameraRigs/Public/CineCameraRigRail.h
--- a/Engine/Plugins/Experimental/CineCameraRigs/Source/CineCameraRigs/Public/CineCameraRigRail.h
+++ b/Engine/Plugins/Experimental/CineCameraRigs/Source/CineCameraRigs/Public/CineCameraRigRail.h
@@ -10,6 +10,8 @@
 
 #include "CineCameraRigRail.generated.h"
 
+#define UE_API CINECAMERARIGS_API
+
 class UMovieSceneSequence;
 class UMovieSceneFloatTrack;
 
@@ -40,15 +42,15 @@ enum class ECineCameraRigRailLoopMode : uint8
 };
 
 UCLASS(Blueprintable, Category = "VirtualProduction")
-class CINECAMERARIGS_API ACineCameraRigRail : public ACameraRig_Rail
+class ACineCameraRigRail : public ACameraRig_Rail
 {
 	GENERATED_BODY()
 
 public:
-	ACineCameraRigRail(const FObjectInitializer& ObjectInitializer);
+	UE_API ACineCameraRigRail(const FObjectInitializer& ObjectInitializer);
 
-	virtual void Tick(float DeltaTime) override;
-	virtual void PostLoad() override;
+	UE_API virtual void Tick(float DeltaTime) override;
+	UE_API virtual void PostLoad() override;
 
 	/* Returns CineSplineComponent*/
 	UFUNCTION(BlueprintPure, Category = "Rail Components")
@@ -151,34 +153,34 @@ public:
 
 	/* Set spline mesh material*/
 	UFUNCTION(BlueprintSetter)
-	void SetSplineMeshMaterial(UMaterialInterface* InMaterial);
+	UE_API void SetSplineMeshMaterial(UMaterialInterface* InMaterial);
 
 	/* Set texture used in the spline mesh material */
 	UFUNCTION(BlueprintSetter)
-	void SetSplineMeshTexture(UTexture2D* InTexture);
+	UE_API void SetSplineMeshTexture(UTexture2D* InTexture);
 
 	/* Calculate internal velocity at the given position */
 	UFUNCTION(BlueprintCallable, Category = "CineCameraRigRail")
-	FVector GetVelocityAtPosition(const float InPosition, const float delta = 0.001) const;
+	UE_API FVector GetVelocityAtPosition(const float InPosition, const float delta = 0.001) const;
 
 	/* Set drive mode*/
 	UFUNCTION(BlueprintSetter)
-	void SetDriveMode(ECineCameraRigRailDriveMode InMode);
+	UE_API void SetDriveMode(ECineCameraRigRailDriveMode InMode);
 
 	/* Enable display speed heatmap*/
 	UFUNCTION(BlueprintSetter)
-	void SetDisplaySpeedHeatmap(bool bEnable);
+	UE_API void SetDisplaySpeedHeatmap(bool bEnable);
 
 	/* Returns true if the rig rail is driven by Sequencer */
 	UFUNCTION()
-	bool IsSequencerDriven();
+	UE_API bool IsSequencerDriven();
 
 	/* Set AbsolutePositionOnRail value*/
 	UFUNCTION(BlueprintSetter)
-	void SetAbsolutePositionOnRail(float Value);
+	UE_API void SetAbsolutePositionOnRail(float Value);
 
 #if WITH_EDITOR
-	virtual void PostEditChangeProperty(FPropertyChangedEvent& PropertyChangedEvent) override;
+	UE_API virtual void PostEditChangeProperty(FPropertyChangedEvent& PropertyChangedEvent) override;
 #endif
 
 protected:
@@ -217,3 +219,5 @@ protected:
 #endif
 
 };
+
+#undef UE_API
\ No newline at end of file
diff --git a/Engine/Plugins/Experimental/CineCameraRigs/Source/CineCameraRigs/Public/CineCameraRigRailHelpers.h b/Engine/Plugins/Experimental/CineCameraRigs/Source/CineCameraRigs/Public/CineCameraRigRailHelpers.h
--- a/Engine/Plugins/Experimental/CineCameraRigs/Source/CineCameraRigs/Public/CineCameraRigRailHelpers.h
+++ b/Engine/Plugins/Experimental/CineCameraRigs/Source/CineCameraRigs/Public/CineCameraRigRailHelpers.h
@@ -8,12 +8,16 @@
 
 #include "CineCameraRigRailHelpers.generated.h"
 
+#define UE_API CINECAMERARIGS_API
+
 UCLASS()
-class CINECAMERARIGS_API UCineCameraRigRailHelpers : public UBlueprintFunctionLibrary
+class UCineCameraRigRailHelpers : public UBlueprintFunctionLibrary
 {
 	GENERATED_BODY()
 public:
 	/* Create a transient heatmap texture from data values*/
 	UFUNCTION(BlueprintCallable, Category = "CineCameraRigRail")
-	static void CreateOrUpdateSplineHeatmapTexture(UPARAM(ref) UTexture2D*& Texture, const TArray<float>& DataValues, const float LowValue, const float AverageValue, const float HighValue);
+	static UE_API void CreateOrUpdateSplineHeatmapTexture(UPARAM(ref) UTexture2D*& Texture, const TArray<float>& DataValues, const float LowValue, const float AverageValue, const float HighValue);
 };
+
+#undef UE_API
\ No newline at end of file
diff --git a/Engine/Plugins/Experimental/CineCameraRigs/Source/CineCameraRigs/Public/CineCameraRigsSettings.h b/Engine/Plugins/Experimental/CineCameraRigs/Source/CineCameraRigs/Public/CineCameraRigsSettings.h
--- a/Engine/Plugins/Experimental/CineCameraRigs/Source/CineCameraRigs/Public/CineCameraRigsSettings.h
+++ b/Engine/Plugins/Experimental/CineCameraRigs/Source/CineCameraRigs/Public/CineCameraRigsSettings.h
@@ -11,7 +11,7 @@ class UMaterialInterface;
 class UTexture2D;
 
 UCLASS(BlueprintType, Config = CineCameraRigs)
-class CINECAMERARIGS_API UCineCameraRigRailSettings : public UObject
+class UCineCameraRigRailSettings : public UObject
 {
 	GENERATED_BODY()
 public:
diff --git a/Engine/Plugins/Experimental/CineCameraRigs/Source/CineCameraRigs/Public/CineSplineComponent.h b/Engine/Plugins/Experimental/CineCameraRigs/Source/CineCameraRigs/Public/CineSplineComponent.h
--- a/Engine/Plugins/Experimental/CineCameraRigs/Source/CineCameraRigs/Public/CineSplineComponent.h
+++ b/Engine/Plugins/Experimental/CineCameraRigs/Source/CineCameraRigs/Public/CineSplineComponent.h
@@ -13,13 +13,15 @@
 
 #include "CineSplineComponent.generated.h"
 
+#define UE_API CINECAMERARIGS_API
+
 UCLASS(ClassGroup = Utility, ShowCategories = (Mobility), HideCategories = (Physics, Collision, Lighting, Rendering, Mobile), meta = (BlueprintSpawnableComponent))
-class CINECAMERARIGS_API UCineSplineComponent : public USplineComponent
+class UCineSplineComponent : public USplineComponent
 {
 	GENERATED_BODY()
 
 public:
-	UCineSplineComponent(const FObjectInitializer& ObjectInitializer);
+	UE_API UCineSplineComponent(const FObjectInitializer& ObjectInitializer);
 
 	/**
 	* Defaults which are used to propagate values to spline points on instances of this in the world
@@ -29,20 +31,20 @@ public:
 
 	// UObject interface
 #if WITH_EDITOR
-	virtual void PostEditChangeProperty(FPropertyChangedEvent& PropertyChangedEvent) override;
+	UE_API virtual void PostEditChangeProperty(FPropertyChangedEvent& PropertyChangedEvent) override;
 #endif
-	virtual void PostLoad() override;
+	UE_API virtual void PostLoad() override;
 	// End of UObject interface
 
 	/** Spline component interface */
-	virtual USplineMetadata* GetSplinePointsMetadata() override;
-	virtual const USplineMetadata* GetSplinePointsMetadata() const override;
+	UE_API virtual USplineMetadata* GetSplinePointsMetadata() override;
+	UE_API virtual const USplineMetadata* GetSplinePointsMetadata() const override;
 	virtual bool AllowsSplinePointScaleEditing() const override { return false; }
 
-	void ApplyComponentInstanceData(struct FCineSplineInstanceData* ComponentInstanceData, const bool bPostUCS);
+	UE_API void ApplyComponentInstanceData(struct FCineSplineInstanceData* ComponentInstanceData, const bool bPostUCS);
 
 	// UActorComponent Interface
-	virtual TStructOnScope<FActorComponentInstanceData> GetComponentInstanceData() const override;
+	UE_API virtual TStructOnScope<FActorComponentInstanceData> GetComponentInstanceData() const override;
 
 	/** Pointer to metadata */
 	UPROPERTY(Instanced)
@@ -50,55 +52,55 @@ public:
 
 	/** Set focal lenght metadata at a given splint point*/
 	UFUNCTION(BlueprintCallable, Category = "CineSpline")
-	void SetFocalLengthAtSplinePoint(const int32 PointIndex, const float Value );
+	UE_API void SetFocalLengthAtSplinePoint(const int32 PointIndex, const float Value );
 
 	/** Set aperture metadata at a given spline point*/
 	UFUNCTION(BlueprintCallable, Category = "CineSpline")
-	void SetApertureAtSplinePoint(const int32 PointIndex, const float value );
+	UE_API void SetApertureAtSplinePoint(const int32 PointIndex, const float value );
 
 	/** Set focus distance metadata at a given spline point*/
 	UFUNCTION(BlueprintCallable, Category = "CineSpline")
-	void SetFocusDistanceAtSplinePoint(const int32 PointIndex, const float value);
+	UE_API void SetFocusDistanceAtSplinePoint(const int32 PointIndex, const float value);
 
 	/** Set custom position metadata at a given spline point*/
 	UFUNCTION(BlueprintCallable, Category = "CineSpline")
-	void SetAbsolutePositionAtSplinePoint(const int32 PointIndex, const float value);
+	UE_API void SetAbsolutePositionAtSplinePoint(const int32 PointIndex, const float value);
 
 	/* Set camera rotation metadata at a given spline point*/
 	UFUNCTION(BlueprintCallable, Category = "CineSpline")
-	void SetPointRotationAtSplinePoint(const int32 PointIndex, const FQuat value);
+	UE_API void SetPointRotationAtSplinePoint(const int32 PointIndex, const FQuat value);
 
 	/** Returns true if there is a spline point at given custom position*/
 	UFUNCTION(BlueprintCallable, BlueprintPure, Category = "CineSpline")
-	bool FindSplineDataAtPosition(const float InPosition, int32& OutIndex, const float Tolerance=0.001f) const;
+	UE_API bool FindSplineDataAtPosition(const float InPosition, int32& OutIndex, const float Tolerance=0.001f) const;
 
 	/* Update spline point data at the given spline point*/
 	UFUNCTION(BlueprintCallable, Category = "CineSpline")
-	void UpdateSplineDataAtIndex(const int InIndex, const FCineSplinePointData& InPointData );
+	UE_API void UpdateSplineDataAtIndex(const int InIndex, const FCineSplinePointData& InPointData );
 
 	/* Add a new spline point data at the given custom position*/
 	UFUNCTION(BlueprintCallable, Category = "CineSpline")
-	void AddSplineDataAtPosition(const float InPosition, const FCineSplinePointData& InPointData);
+	UE_API void AddSplineDataAtPosition(const float InPosition, const FCineSplinePointData& InPointData);
 
 	/* Get a spline point data at the given custom position*/
 	UFUNCTION(BlueprintCallable, Category = "CineSpline")
-	FCineSplinePointData GetSplineDataAtPosition(const float InPosition) const;
+	UE_API FCineSplinePointData GetSplineDataAtPosition(const float InPosition) const;
 
 	/* Get input key value from custom position */
 	UFUNCTION(BlueprintCallable, Category = "CineSpline")
-	float GetInputKeyAtPosition(const float InPosition) const;
+	UE_API float GetInputKeyAtPosition(const float InPosition) const;
 
 	/* Get custom position value at spline input key */
 	UFUNCTION(BlueprintCallable, Category = "CineSpline")
-	float GetPositionAtInputKey(const float InKey) const;
+	UE_API float GetPositionAtInputKey(const float InKey) const;
 
 	/* Get camera rotation metadata property value along the spline at spline point */
 	UFUNCTION(BlueprintCallable, Category = "CineSpline")
-	FQuat GetPointRotationAtSplinePoint(int32 Index) const;
+	UE_API FQuat GetPointRotationAtSplinePoint(int32 Index) const;
 
 	/* Get camera rotation metadata value along the spline at spline input key*/
 	UFUNCTION(BlueprintCallable, Category = "CineSpline")
-	FQuat GetPointRotationAtSplineInputKey(float InKey) const;
+	UE_API FQuat GetPointRotationAtSplineInputKey(float InKey) const;
 
 
 	DECLARE_MULTICAST_DELEGATE(FOnSplineEdited);
@@ -111,7 +113,7 @@ public:
 	FOnSplineEdited_BP OnSplineEdited_BP;
 
 	/** Update the spline tangents and SplineReparamTable. And trigers OnSplineEdited delegates */
-	virtual void UpdateSpline() override;
+	UE_API virtual void UpdateSpline() override;
 
 #if WITH_EDITORONLY_DATA
 	/* Toggle if visualizer displays normalized position or absolute position */
@@ -145,4 +147,6 @@ public:
 
 	UPROPERTY()
 	TObjectPtr<UCineSplineMetadata> CineSplineMetadata;
-};
\ No newline at end of file
+};
+
+#undef UE_API
\ No newline at end of file
diff --git a/Engine/Plugins/Experimental/CineCameraRigs/Source/CineCameraRigs/Public/CineSplineMetadata.h b/Engine/Plugins/Experimental/CineCameraRigs/Source/CineCameraRigs/Public/CineSplineMetadata.h
--- a/Engine/Plugins/Experimental/CineCameraRigs/Source/CineCameraRigs/Public/CineSplineMetadata.h
+++ b/Engine/Plugins/Experimental/CineCameraRigs/Source/CineCameraRigs/Public/CineSplineMetadata.h
@@ -9,6 +9,8 @@
 
 #include "CineSplineMetadata.generated.h"
 
+#define UE_API CINECAMERARIGS_API
+
 USTRUCT()
 struct FCineSplineCurveDefaults
 {
@@ -39,21 +41,21 @@ struct FCineSplineCurveDefaults
 };
 
 UCLASS()
-class CINECAMERARIGS_API UCineSplineMetadata : public USplineMetadata
+class UCineSplineMetadata : public USplineMetadata
 {
 	GENERATED_BODY()
 	
 public:
 	/** Insert point before index, lerping metadata between previous and next key values */
-	virtual void InsertPoint(int32 Index, float Alpha, bool bClosedLoop) override;
+	UE_API virtual void InsertPoint(int32 Index, float Alpha, bool bClosedLoop) override;
 	/** Update point at index by lerping metadata between previous and next key values */
-	virtual void UpdatePoint(int32 Index, float Alpha, bool bClosedLoop) override;
-	virtual void AddPoint(float InputKey) override;
-	virtual void RemovePoint(int32 Index) override;
-	virtual void DuplicatePoint(int32 Index) override;
-	virtual void CopyPoint(const USplineMetadata* FromSplineMetadata, int32 FromIndex, int32 ToIndex) override;
-	virtual void Reset(int32 NumPoints) override;
-	virtual void Fixup(int32 NumPoints, USplineComponent* SplineComp) override;
+	UE_API virtual void UpdatePoint(int32 Index, float Alpha, bool bClosedLoop) override;
+	UE_API virtual void AddPoint(float InputKey) override;
+	UE_API virtual void RemovePoint(int32 Index) override;
+	UE_API virtual void DuplicatePoint(int32 Index) override;
+	UE_API virtual void CopyPoint(const USplineMetadata* FromSplineMetadata, int32 FromIndex, int32 ToIndex) override;
+	UE_API virtual void Reset(int32 NumPoints) override;
+	UE_API virtual void Fixup(int32 NumPoints, USplineComponent* SplineComp) override;
 
 	UPROPERTY(EditAnywhere, Category = "Point")
 	FInterpCurveFloat AbsolutePosition;
@@ -70,3 +72,5 @@ public:
 	UPROPERTY(EditAnywhere, Category = "Camera")
 	FInterpCurveQuat PointRotation;
 };
+
+#undef UE_API
\ No newline at end of file
diff --git a/Engine/Plugins/Experimental/CineCameraRigs/Source/CineCameraRigsEditor/Private/CineCameraRigRailDetails.cpp b/Engine/Plugins/Experimental/CineCameraRigs/Source/CineCameraRigsEditor/Private/CineCameraRigRailDetails.cpp
--- a/Engine/Plugins/Experimental/CineCameraRigs/Source/CineCameraRigsEditor/Private/CineCameraRigRailDetails.cpp
+++ b/Engine/Plugins/Experimental/CineCameraRigs/Source/CineCameraRigsEditor/Private/CineCameraRigRailDetails.cpp
@@ -2,6 +2,7 @@
 
 #include "CineCameraRigRailDetails.h"
 #include "CineCameraRigRail.h"
+#include "CineSplineComponent.h"
 #include "CameraRig_Rail.h"
 #include "DetailCategoryBuilder.h"
 #include "DetailLayoutBuilder.h"
diff --git a/Engine/Plugins/Experimental/CineCameraRigs/Source/CineCameraRigsEditor/Private/CineSplineComponentVisualizer.cpp b/Engine/Plugins/Experimental/CineCameraRigs/Source/CineCameraRigsEditor/Private/CineSplineComponentVisualizer.cpp
--- a/Engine/Plugins/Experimental/CineCameraRigs/Source/CineCameraRigsEditor/Private/CineSplineComponentVisualizer.cpp
+++ b/Engine/Plugins/Experimental/CineCameraRigs/Source/CineCameraRigsEditor/Private/CineSplineComponentVisualizer.cpp
@@ -1,6 +1,7 @@
 // Copyright Epic Games, Inc. All Rights Reserved.
 
 #include "CineSplineComponentVisualizer.h"
+#include "CineSplineComponent.h"
 #include "Internationalization/Text.h"
 #include "Math/Color.h"
 #include "CanvasTypes.h"
