[CineCameraRigTail]
- Move TimerManager registration to PostLoad during loading to avoid getting called before the object is ready

diff --git a/Engine/Plugins/Experimental/CineCameraRigs/Source/CineCameraRigs/Private/CineCameraRigRail.cpp b/Engine/Plugins/Experimental/CineCameraRigs/Source/CineCameraRigs/Private/CineCameraRigRail.cpp
--- a/Engine/Plugins/Experimental/CineCameraRigs/Source/CineCameraRigs/Private/CineCameraRigRail.cpp
+++ b/Engine/Plugins/Experimental/CineCameraRigs/Source/CineCameraRigs/Private/CineCameraRigRail.cpp
@@ -66,13 +66,23 @@ ACineCameraRigRail::ACineCameraRigRail(const FObjectInitializer& ObjectInitializ
 #endif
 
 #if WITH_EDITOR
-	if (!IsTemplate() && GEditor)
+	if (!IsInAsyncLoadingThread())
 	{
-		GEditor->GetTimerManager()->SetTimer(SequencerCheckHandle, this, &ACineCameraRigRail::OnSequencerCheck, 1.0, true);
+		ConditionalRegisterTimerManager();
 	}
 #endif
 }
 
+#if WITH_EDITOR
+void ACineCameraRigRail::ConditionalRegisterTimerManager()
+{
+	if (!SequencerCheckHandle.IsValid() && !IsTemplate() && GEditor)
+	{
+		GEditor->GetTimerManager()->SetTimer(SequencerCheckHandle, this, &ACineCameraRigRail::OnSequencerCheck, 1.0, true);
+	}
+}
+#endif
+
 void ACineCameraRigRail::UpdateRailComponents()
 {
 	if (CineSplineComponent == nullptr)
@@ -370,6 +380,10 @@ void ACineCameraRigRail::PostLoad()
 		UpdateSplineMeshMID();
 		SetMIDParameters();
 	}
+
+#if WITH_EDITOR
+	ConditionalRegisterTimerManager();
+#endif
 }
 
 void ACineCameraRigRail::DriveByParam(float DeltaTime)
diff --git a/Engine/Plugins/Experimental/CineCameraRigs/Source/CineCameraRigs/Public/CineCameraRigRail.h b/Engine/Plugins/Experimental/CineCameraRigs/Source/CineCameraRigs/Public/CineCameraRigRail.h
--- a/Engine/Plugins/Experimental/CineCameraRigs/Source/CineCameraRigs/Public/CineCameraRigRail.h
+++ b/Engine/Plugins/Experimental/CineCameraRigs/Source/CineCameraRigs/Public/CineCameraRigRail.h
@@ -207,6 +207,7 @@ protected:
 	bool PlayForward() const;
 
 #if WITH_EDITOR
+	void ConditionalRegisterTimerManager();
 	UMovieSceneFloatTrack* FindPositionTrack(const UMovieSceneSequence* InSequence);
 
 	FTimerHandle SequencerCheckHandle;
