Adding metadata interpolation option for CineCameraRigRail
#Synced-CL 48811307

diff --git a/Engine/Plugins/Experimental/CineCameraRigs/Source/CineCameraRigs/Private/CineSplineComponent.cpp b/Engine/Plugins/Experimental/CineCameraRigs/Source/CineCameraRigs/Private/CineSplineComponent.cpp
--- a/Engine/Plugins/Experimental/CineCameraRigs/Source/CineCameraRigs/Private/CineSplineComponent.cpp
+++ b/Engine/Plugins/Experimental/CineCameraRigs/Source/CineCameraRigs/Private/CineSplineComponent.cpp
@@ -289,6 +289,10 @@ FCineSplinePointData UCineSplineComponent::GetSplineDataAtPosition(const float I
 void UCineSplineComponent::UpdateSpline()
 {
 	USplineComponent::UpdateSpline();
+	if (CineSplineMetadata)
+	{
+		CineSplineMetadata->AutoSetTangents();
+	}
 	if (OnSplineEdited.IsBound())
 	{
 		OnSplineEdited.Broadcast();
diff --git a/Engine/Plugins/Experimental/CineCameraRigs/Source/CineCameraRigs/Private/CineSplineMetadata.cpp b/Engine/Plugins/Experimental/CineCameraRigs/Source/CineCameraRigs/Private/CineSplineMetadata.cpp
--- a/Engine/Plugins/Experimental/CineCameraRigs/Source/CineCameraRigs/Private/CineSplineMetadata.cpp
+++ b/Engine/Plugins/Experimental/CineCameraRigs/Source/CineCameraRigs/Private/CineSplineMetadata.cpp
@@ -68,6 +68,7 @@ void UCineSplineMetadata::InsertPoint(int32 Index, float Alpha, bool bClosedLoop
 			AbsolutePosition.Points[i].InVal += 1.0f;
 			PointRotation.Points[i].InVal += 1.0f;
 		}
+		AutoSetTangents();
 	}
 }
 
@@ -108,6 +109,7 @@ void UCineSplineMetadata::UpdatePoint(int32 Index, float Alpha, bool bClosedLoop
 		FocusDistance.Points[Index].OutVal = FMath::LerpStable(PrevFocusDistanceVal, NextFocusDistanceVal, Alpha);
 		PointRotation.Points[Index].OutVal = FMath::LerpStable(PrevRotationVal, NextRotationVal, Alpha);
 	}
+	AutoSetTangents();
 }
 
 void UCineSplineMetadata::AddPoint(float InputKey)
@@ -136,6 +138,8 @@ void UCineSplineMetadata::AddPoint(float InputKey)
 	FocusDistance.Points.Emplace(NewInputKey, NewFocusDistanceVal);
 	AbsolutePosition.Points.Emplace(NewInputKey, NewAbsolutePositionVal);
 	PointRotation.Points.Emplace(NewInputKey, NewRotationVal);
+
+	AutoSetTangents();
 }
 
 void UCineSplineMetadata::RemovePoint(int32 Index)
@@ -157,6 +161,8 @@ void UCineSplineMetadata::RemovePoint(int32 Index)
 		FocusDistance.Points[i].InVal -= 1.0f;
 		PointRotation.Points[i].InVal -= 1.0f;
 	}
+
+	AutoSetTangents();
 }
 
 void UCineSplineMetadata::DuplicatePoint(int32 Index)
@@ -187,7 +193,7 @@ void UCineSplineMetadata::DuplicatePoint(int32 Index)
 		FocusDistance.Points[i].InVal += 1.0f;
 		PointRotation.Points[i].InVal += 1.0f;
 	}
-
+	AutoSetTangents();
 }
 
 void UCineSplineMetadata::CopyPoint(const USplineMetadata* FromSplineMetadata, int32 FromIndex, int32 ToIndex)
@@ -205,6 +211,8 @@ void UCineSplineMetadata::CopyPoint(const USplineMetadata* FromSplineMetadata, i
 		FocusDistance.Points[ToIndex].OutVal = FromMetadata->FocusDistance.Points[FromIndex].OutVal;
 		AbsolutePosition.Points[ToIndex].OutVal = FromMetadata->AbsolutePosition.Points[FromIndex].OutVal;
 		PointRotation.Points[ToIndex].OutVal = FromMetadata->PointRotation.Points[FromIndex].OutVal;
+
+		AutoSetTangents();
 	}
 }
 
@@ -218,6 +226,67 @@ void UCineSplineMetadata::Reset(int32 NumPoints)
 	PointRotation.Points.Reset(NumPoints);
 }
 
+void UCineSplineMetadata::AutoSetTangents()
+{
+	Modify();
+	FocalLength.AutoSetTangents(0.0f, false);
+	Aperture.AutoSetTangents(0.0f, false);
+	FocusDistance.AutoSetTangents(0.0f, false);
+	PointRotation.AutoSetTangents(0.0f, false);
+}
+
+void UCineSplineMetadata::SetTangentType(int32 Index, ECineSplineMetadataTangentType TangentType)
+{
+	check(Index < FocalLength.Points.Num());
+	check(Index >= 0);
+
+	EInterpCurveMode Mode = ConvertCineSplineMetadataTangentTypeToInterpCurveMode(TangentType);
+	if (Mode != CIM_Unknown)
+	{
+		Modify();
+		FocalLength.Points[Index].InterpMode = Mode;
+		Aperture.Points[Index].InterpMode = Mode;
+		FocusDistance.Points[Index].InterpMode = Mode;
+		PointRotation.Points[Index].InterpMode = Mode;
+	}
+}
+
+EInterpCurveMode ConvertCineSplineMetadataTangentTypeToInterpCurveMode(ECineSplineMetadataTangentType MetadataTangentType)
+{
+	switch (MetadataTangentType)
+	{
+		case ECineSplineMetadataTangentType::Linear:
+			return CIM_Linear;
+		case ECineSplineMetadataTangentType::Curve:
+			return CIM_CurveAuto;
+		case ECineSplineMetadataTangentType::Constant:
+			return CIM_Constant;
+		case ECineSplineMetadataTangentType::CurveClamped:
+			return CIM_CurveAutoClamped;
+
+		default:
+			return CIM_Unknown;
+	}
+}
+
+ECineSplineMetadataTangentType ConvertInterpCurveModeToCineSplineMetadataTangentType(EInterpCurveMode InterpCurveMode)
+{
+	switch (InterpCurveMode)
+	{
+		case CIM_Linear:
+			return ECineSplineMetadataTangentType::Linear;
+		case CIM_CurveAuto:
+			return ECineSplineMetadataTangentType::Curve;
+		case CIM_Constant:
+			return ECineSplineMetadataTangentType::Constant;
+		case CIM_CurveAutoClamped:
+			return ECineSplineMetadataTangentType::CurveClamped;
+
+		default:
+			return ECineSplineMetadataTangentType::Linear;
+	}
+}
+
 #if WITH_EDITORONLY_DATA
 template <class T>
 void FixupCurve(FInterpCurve<T>& Curve, const T& DefaultValue, int32 NumPoints)
diff --git a/Engine/Plugins/Experimental/CineCameraRigs/Source/CineCameraRigs/Public/CineSplineMetadata.h b/Engine/Plugins/Experimental/CineCameraRigs/Source/CineCameraRigs/Public/CineSplineMetadata.h
--- a/Engine/Plugins/Experimental/CineCameraRigs/Source/CineCameraRigs/Public/CineSplineMetadata.h
+++ b/Engine/Plugins/Experimental/CineCameraRigs/Source/CineCameraRigs/Public/CineSplineMetadata.h
@@ -40,6 +40,22 @@ struct FCineSplineCurveDefaults
 	FQuat DefaultPointRotation;
 };
 
+/** Interpolation type for CineSpline Metadata */
+UENUM()
+enum class ECineSplineMetadataTangentType : uint8
+{
+	Linear,
+	Curve,
+	Constant,
+	CurveClamped
+};
+
+/** Convert ECineSplineMetadataTangentType value to EInterpCurveMode value */
+UE_API EInterpCurveMode ConvertCineSplineMetadataTangentTypeToInterpCurveMode(ECineSplineMetadataTangentType MetadataTangentType);
+
+/** Convert EInterpCurveMode value to ECineSplineMetadataTangentType value */
+UE_API ECineSplineMetadataTangentType ConvertInterpCurveModeToCineSplineMetadataTangentType(EInterpCurveMode InterpCurveMode);
+
 UCLASS()
 class UCineSplineMetadata : public USplineMetadata
 {
@@ -71,6 +87,12 @@ public:
 
 	UPROPERTY(EditAnywhere, Category = "Camera")
 	FInterpCurveQuat PointRotation;
+
+	/** Automatically set the tangents on the metadata parameters based on surrounding points */
+	UE_API void AutoSetTangents();
+
+	/** Specify tangent type on the spline point metadata */
+	UE_API void SetTangentType(int32 Index, ECineSplineMetadataTangentType TangentType);
 };
 
 #undef UE_API
\ No newline at end of file
diff --git a/Engine/Plugins/Experimental/CineCameraRigs/Source/CineCameraRigsEditor/Private/CineSplineComponentVisualizer.cpp b/Engine/Plugins/Experimental/CineCameraRigs/Source/CineCameraRigsEditor/Private/CineSplineComponentVisualizer.cpp
--- a/Engine/Plugins/Experimental/CineCameraRigs/Source/CineCameraRigsEditor/Private/CineSplineComponentVisualizer.cpp
+++ b/Engine/Plugins/Experimental/CineCameraRigs/Source/CineCameraRigsEditor/Private/CineSplineComponentVisualizer.cpp
@@ -33,6 +33,10 @@ public:
 		UI_COMMAND(VisualizeSplineLength, "Show Spline Length", "Whether the visualization should show spline length.", EUserInterfaceActionType::ToggleButton, FInputChord());
 		UI_COMMAND(VisualizePointRotation, "Show Point Rotation", "Whether the visualization should show point rotation arrows.", EUserInterfaceActionType::ToggleButton, FInputChord());
 		UI_COMMAND(ManipulatePointRotation, "Manipulate Point Rotation", "Use rotation manipulator to edit PointRotation metadata.", EUserInterfaceActionType::ToggleButton, FInputChord());
+		UI_COMMAND(SetMetadataToLinear, "Linear", "Set spline point to Linear type", EUserInterfaceActionType::RadioButton, FInputChord());
+		UI_COMMAND(SetMetadataToCurve, "Curve", "Set spline point to Curve type", EUserInterfaceActionType::RadioButton, FInputChord());
+		UI_COMMAND(SetMetadataToConstant, "Constant", "Set spline point to Constant type", EUserInterfaceActionType::RadioButton, FInputChord());
+		UI_COMMAND(SetMetadataToCurveClamped, "CurveClamped", "Set spline point to CurveClamped type", EUserInterfaceActionType::RadioButton, FInputChord());
 	}
 
 public:
@@ -41,6 +45,10 @@ public:
 	TSharedPtr<FUICommandInfo> VisualizeSplineLength;
 	TSharedPtr<FUICommandInfo> VisualizePointRotation;
 	TSharedPtr<FUICommandInfo> ManipulatePointRotation;
+	TSharedPtr<FUICommandInfo> SetMetadataToLinear;
+	TSharedPtr<FUICommandInfo> SetMetadataToCurve;
+	TSharedPtr<FUICommandInfo> SetMetadataToConstant;
+	TSharedPtr<FUICommandInfo> SetMetadataToCurveClamped;
 };
 
 FCineSplineComponentVisualizer::FCineSplineComponentVisualizer()
@@ -83,6 +91,31 @@ void FCineSplineComponentVisualizer::OnRegister()
 		FCanExecuteAction(),
 		FIsActionChecked::CreateSP(this, &FCineSplineComponentVisualizer::IsManipulatePointRotation)
 	);
+
+	SplineComponentVisualizerActions->MapAction(
+		Commands.SetMetadataToLinear,
+		FExecuteAction::CreateSP(this, &FCineSplineComponentVisualizer::OnSetMetadataTangent, ECineSplineMetadataTangentType::Linear),
+		FCanExecuteAction(),
+		FIsActionChecked::CreateSP(this, &FCineSplineComponentVisualizer::IsMetadataTangent, ECineSplineMetadataTangentType::Linear)
+	);
+	SplineComponentVisualizerActions->MapAction(
+		Commands.SetMetadataToCurve,
+		FExecuteAction::CreateSP(this, &FCineSplineComponentVisualizer::OnSetMetadataTangent, ECineSplineMetadataTangentType::Curve),
+		FCanExecuteAction(),
+		FIsActionChecked::CreateSP(this, &FCineSplineComponentVisualizer::IsMetadataTangent, ECineSplineMetadataTangentType::Curve)
+	);
+	SplineComponentVisualizerActions->MapAction(
+		Commands.SetMetadataToConstant,
+		FExecuteAction::CreateSP(this, &FCineSplineComponentVisualizer::OnSetMetadataTangent, ECineSplineMetadataTangentType::Constant),
+		FCanExecuteAction(),
+		FIsActionChecked::CreateSP(this, &FCineSplineComponentVisualizer::IsMetadataTangent, ECineSplineMetadataTangentType::Constant)
+	);
+	SplineComponentVisualizerActions->MapAction(
+		Commands.SetMetadataToCurveClamped,
+		FExecuteAction::CreateSP(this, &FCineSplineComponentVisualizer::OnSetMetadataTangent, ECineSplineMetadataTangentType::CurveClamped),
+		FCanExecuteAction(),
+		FIsActionChecked::CreateSP(this, &FCineSplineComponentVisualizer::IsMetadataTangent, ECineSplineMetadataTangentType::CurveClamped)
+	);
 }
 
 FCineSplineComponentVisualizer::~FCineSplineComponentVisualizer()
@@ -179,16 +212,22 @@ void FCineSplineComponentVisualizer::DrawVisualization(
 void FCineSplineComponentVisualizer::GenerateContextMenuSections(FMenuBuilder& InMenuBuilder) const
 {
 	FSplineComponentVisualizer::GenerateContextMenuSections(InMenuBuilder);
-
-	InMenuBuilder.BeginSection("CineSpline", LOCTEXT("CineSpline", "CineSpline"));
+	if (GetEditedSplineComponent() != nullptr)
 	{
-		InMenuBuilder.AddMenuEntry(FCineSplineComponentVisualizerCommands::Get().VisualizeNormalizedPosition);
-		InMenuBuilder.AddMenuEntry(FCineSplineComponentVisualizerCommands::Get().VisualizeSplineLength);
-		InMenuBuilder.AddMenuEntry(FCineSplineComponentVisualizerCommands::Get().VisualizePointRotation);
-		InMenuBuilder.AddMenuEntry(FCineSplineComponentVisualizerCommands::Get().ManipulatePointRotation);
+		InMenuBuilder.BeginSection("CineSpline", LOCTEXT("CineSpline", "CineSpline"));
+		{
+			InMenuBuilder.AddMenuEntry(FCineSplineComponentVisualizerCommands::Get().VisualizeNormalizedPosition);
+			InMenuBuilder.AddMenuEntry(FCineSplineComponentVisualizerCommands::Get().VisualizeSplineLength);
+			InMenuBuilder.AddMenuEntry(FCineSplineComponentVisualizerCommands::Get().VisualizePointRotation);
+			InMenuBuilder.AddMenuEntry(FCineSplineComponentVisualizerCommands::Get().ManipulatePointRotation);
+			InMenuBuilder.AddSubMenu(
+				LOCTEXT("CineSplineMetadataTangentType", "Metadata Tangent Type"),
+				LOCTEXT("CineSplineMetadataTangentTypeTooltip", "Define the type of the metadata tangent"),
+				FNewMenuDelegate::CreateSP(this, &FCineSplineComponentVisualizer::GenerateMetadataTangentTypeSubMenu)
+			);
+		}
+		InMenuBuilder.EndSection();
 	}
-	InMenuBuilder.EndSection();
-
 }
 
 bool FCineSplineComponentVisualizer::HandleInputDelta(FEditorViewportClient* ViewportClient, FViewport* Viewport, FVector& DeltaTranslate, FRotator& DeltaRotate, FVector& DeltaScale)
@@ -307,4 +346,60 @@ bool FCineSplineComponentVisualizer::IsManipulatePointRotation() const
 	return bManipulatePointRotation;
 }
 
+void FCineSplineComponentVisualizer::GenerateMetadataTangentTypeSubMenu(FMenuBuilder& MenuBuilder) const
+{
+	if (const UCineSplineComponent* SplineComp = Cast<UCineSplineComponent>(GetEditedSplineComponent()))
+	{
+		MenuBuilder.AddMenuEntry(FCineSplineComponentVisualizerCommands::Get().SetMetadataToLinear);
+		MenuBuilder.AddMenuEntry(FCineSplineComponentVisualizerCommands::Get().SetMetadataToCurve);
+		MenuBuilder.AddMenuEntry(FCineSplineComponentVisualizerCommands::Get().SetMetadataToConstant);
+		MenuBuilder.AddMenuEntry(FCineSplineComponentVisualizerCommands::Get().SetMetadataToCurveClamped);
+	}
+}
+
+void FCineSplineComponentVisualizer::OnSetMetadataTangent(ECineSplineMetadataTangentType TangentType)
+{
+	if (UCineSplineComponent* SplineComp = Cast<UCineSplineComponent>(GetEditedSplineComponent()))
+	{
+		check(SelectionState);
+		const TSet<int32>& SelectedKeys = SelectionState->GetSelectedKeys();
+		for (int32 SelectedKeyIndex : SelectedKeys)
+		{
+			check(SelectedKeyIndex >= 0);
+			check(SelectedKeyIndex < SplineComp->GetNumberOfSplinePoints());
+			SplineComp->CineSplineMetadata->SetTangentType(SelectedKeyIndex, TangentType);
+		}
+		SplineComp->CineSplineMetadata->AutoSetTangents();
+	}
+}
+
+bool FCineSplineComponentVisualizer::IsMetadataTangent(ECineSplineMetadataTangentType TangentType)
+{
+	if (IsKeySelectionValid())
+	{
+		UCineSplineComponent* SplineComp = Cast<UCineSplineComponent>(GetEditedSplineComponent());
+		check(SplineComp != nullptr);
+		check(SelectionState);
+
+		EInterpCurveMode Mode = ConvertCineSplineMetadataTangentTypeToInterpCurveMode(TangentType);
+		if(Mode != CIM_Unknown)
+		{
+			const TSet<int32>& SelectedKeys = SelectionState->GetSelectedKeys();
+			for (int32 SelectedKeyIndex : SelectedKeys)
+			{
+				check(SelectedKeyIndex >= 0);
+				check(SelectedKeyIndex < SplineComp->GetNumberOfSplinePoints());
+				const auto& SelectedPoint = SplineComp->CineSplineMetadata->FocalLength.Points[SelectedKeyIndex];
+				if (SelectedPoint.InterpMode == Mode)
+				{
+					return true;
+				}
+			}
+		}
+	}
+
+	return false;
+}
+
+
 #undef LOCTEXT_NAMESPACE
\ No newline at end of file
diff --git a/Engine/Plugins/Experimental/CineCameraRigs/Source/CineCameraRigsEditor/Private/CineSplineComponentVisualizer.h b/Engine/Plugins/Experimental/CineCameraRigs/Source/CineCameraRigsEditor/Private/CineSplineComponentVisualizer.h
--- a/Engine/Plugins/Experimental/CineCameraRigs/Source/CineCameraRigsEditor/Private/CineSplineComponentVisualizer.h
+++ b/Engine/Plugins/Experimental/CineCameraRigs/Source/CineCameraRigsEditor/Private/CineSplineComponentVisualizer.h
@@ -12,6 +12,7 @@
 
 #include "Components/SplineComponent.h"
 #include "CineSplineComponent.h"
+#include "CineSplineMetadata.h"
 
 class FUICommandList;
 
@@ -43,6 +44,10 @@ protected:
 	void OnSetManipulatePointRotation();
 	bool IsManipulatePointRotation() const;
 
+	void OnSetMetadataTangent(ECineSplineMetadataTangentType TangentType);
+	bool IsMetadataTangent(ECineSplineMetadataTangentType TangentType);
+	void GenerateMetadataTangentTypeSubMenu(FMenuBuilder& MenuBuilder) const;
+
 	TSharedPtr<FUICommandList> CineSplineComponentVisualizerActions;
 
 };
\ No newline at end of file
diff --git a/Engine/Plugins/Experimental/CineCameraRigs/Source/CineCameraRigsEditor/Private/CineSplineMetadataDetails.cpp b/Engine/Plugins/Experimental/CineCameraRigs/Source/CineCameraRigsEditor/Private/CineSplineMetadataDetails.cpp
--- a/Engine/Plugins/Experimental/CineCameraRigs/Source/CineCameraRigsEditor/Private/CineSplineMetadataDetails.cpp
+++ b/Engine/Plugins/Experimental/CineCameraRigs/Source/CineCameraRigsEditor/Private/CineSplineMetadataDetails.cpp
@@ -58,6 +58,7 @@ void FCineSplineMetadataDetails::Update(USplineComponent* InSplineComponent, con
 	ApertureValue.Reset();
 	FocusDistanceValue.Reset();
 	PointRotationValue.Reset();
+	InterpTypeValue.Reset();
 
 	if (!IsValid(InSplineComponent))
 	{
@@ -74,6 +75,7 @@ void FCineSplineMetadataDetails::Update(USplineComponent* InSplineComponent, con
 	bool bUpdateAperture = true;
 	bool bUpdateFocusDistance = true;
 	bool bUpdatePointRotation = true;
+	bool bUpdateInterpType = true;
 
 	for (const int32 Index : InSelectedKeys)
 	{
@@ -82,6 +84,7 @@ void FCineSplineMetadataDetails::Update(USplineComponent* InSplineComponent, con
 		bUpdateAperture = bUpdateAperture ? UpdateMultipleValue(ApertureValue, Metadata->Aperture.Points[Index].OutVal) : false;
 		bUpdateFocusDistance = bUpdateFocusDistance ? UpdateMultipleValue(FocusDistanceValue, Metadata->FocusDistance.Points[Index].OutVal) : false;
 		bUpdatePointRotation = bUpdatePointRotation ? UpdateMultipleValue(PointRotationValue, Metadata->PointRotation.Points[Index].OutVal) : false;
+		bUpdateInterpType = bUpdateInterpType ? UpdateMultipleValue(InterpTypeValue, ConvertInterpCurveModeToCineSplineMetadataTangentType(Metadata->PointRotation.Points[Index].InterpMode)) : false;
 	}
 }
 
@@ -192,6 +195,77 @@ void FCineSplineMetadataDetails::OnSetRotation(float NewValue, ETextCommit::Type
 	}
 }
 
+FText FCineSplineMetadataDetails::GetInterpType() const
+{
+	if (InterpTypeValue.IsSet())
+	{
+		const UEnum* InterpTypeEnum = StaticEnum<ECineSplineMetadataTangentType>();
+		return InterpTypeEnum->GetDisplayNameTextByValue((int64)InterpTypeValue.GetValue());
+	}
+	return LOCTEXT("MultipleTypes", "Multiple Types");
+}
+
+void FCineSplineMetadataDetails::OnInterpTypeChanged(TSharedPtr<FString> NewValue, ESelectInfo::Type SelectInfo)
+{
+	if (SelectInfo == ESelectInfo::Direct)
+	{
+		return;
+	}
+
+	if (UCineSplineMetadata* Metadata = GetMetadata())
+	{
+		const UEnum* InterpTypeEnum = StaticEnum<ECineSplineMetadataTangentType>();
+		const int64 InterpType =InterpTypeEnum->GetValueByNameString(*NewValue);
+		if (InterpType != INDEX_NONE)
+		{
+			EInterpCurveMode Mode = ConvertCineSplineMetadataTangentTypeToInterpCurveMode(static_cast<ECineSplineMetadataTangentType>(InterpType));
+
+			const FScopedTransaction Transaction(LOCTEXT("SetSplineMetadataInterpType", "Set spline metadata interp type"));
+			Metadata->Modify();
+			for (int32 Index : SelectedKeys)
+			{
+				Metadata->SetTangentType(Index, static_cast<ECineSplineMetadataTangentType>(InterpType));
+			}
+			Metadata->AutoSetTangents();
+
+			if (GEditor)
+			{
+				GEditor->RedrawLevelEditingViewports(true);
+			}
+		}
+	}
+}
+
+void FCineSplineMetadataDetails::UpdateInterpTypes()
+{
+	const UEnum* InterpTypeEnum = StaticEnum<ECineSplineMetadataTangentType>();
+	InterpTypes.Reset();
+	for (int32 Index = 0; Index < InterpTypeEnum->NumEnums() - 1; ++Index)
+	{
+		InterpTypes.Add(MakeShareable(new FString(InterpTypeEnum->GetNameStringByIndex(Index))));
+	}
+}
+
+TSharedRef<SWidget> FCineSplineMetadataDetails::OnGenerateInterpTypeComboWidget(TSharedPtr<FString> InComboString)
+{
+	return SNew(STextBlock)
+		.Text(FText::FromString(*InComboString))
+		.Font(IDetailLayoutBuilder::GetDetailFont());
+}
+
+void FCineSplineMetadataDetails::OnInterpTypeComboWidgetOpened()
+{
+	TSharedPtr<FString> SelectedItem;
+	for (TSharedPtr<FString> Item : InterpTypes)
+	{
+		if (Item->Equals(GetInterpType().ToString()))
+		{
+			SelectedItem = Item;
+		}
+	}
+	InterpTypeComboBox->SetSelectedItem(SelectedItem);
+}
+
 void FCineSplineMetadataDetails::GenerateChildContent(IDetailGroup& DetailGroup)
 {
 	DetailGroup.AddWidgetRow()
@@ -347,7 +421,34 @@ void FCineSplineMetadataDetails::GenerateChildContent(IDetailGroup& DetailGroup)
 			.Font(IDetailLayoutBuilder::GetDetailFont())
 		];
 
+	UpdateInterpTypes();
 
+	DetailGroup.AddWidgetRow()
+		.RowTag("InterpType")
+		.Visibility(TAttribute<EVisibility>(this, &FCineSplineMetadataDetails::IsEnabled))
+		.NameContent()
+		.HAlign(HAlign_Left)
+		.VAlign(VAlign_Center)
+		[
+			SNew(STextBlock)
+				.Text(LOCTEXT("InterpType", "Interp Type"))
+				.Font(IDetailLayoutBuilder::GetDetailFont())
+		]
+		.ValueContent()
+		.MinDesiredWidth(125.0f)
+		.MaxDesiredWidth(125.0f)
+		[
+			SAssignNew(InterpTypeComboBox, SComboBox<TSharedPtr<FString>>)
+				.OptionsSource(&InterpTypes)
+				.OnGenerateWidget(this, &FCineSplineMetadataDetails::OnGenerateInterpTypeComboWidget)
+				.OnSelectionChanged(this, &FCineSplineMetadataDetails::OnInterpTypeChanged)
+				.OnComboBoxOpening(this, &FCineSplineMetadataDetails::OnInterpTypeComboWidgetOpened)
+				[
+					SNew(STextBlock)
+						.Font(IDetailLayoutBuilder::GetDetailFont())
+						.Text(this, &FCineSplineMetadataDetails::GetInterpType)
+				]
+		];
 }
 
 #undef LOCTEXT_NAMESPACE
\ No newline at end of file
diff --git a/Engine/Plugins/Experimental/CineCameraRigs/Source/CineCameraRigsEditor/Private/CineSplineMetadataDetails.h b/Engine/Plugins/Experimental/CineCameraRigs/Source/CineCameraRigsEditor/Private/CineSplineMetadataDetails.h
--- a/Engine/Plugins/Experimental/CineCameraRigs/Source/CineCameraRigsEditor/Private/CineSplineMetadataDetails.h
+++ b/Engine/Plugins/Experimental/CineCameraRigs/Source/CineCameraRigsEditor/Private/CineSplineMetadataDetails.h
@@ -7,6 +7,8 @@
 #include "Layout/Visibility.h"
 #include "Types/SlateEnums.h"
 #include "ScopedTransaction.h"
+#include "CineSplineMetadata.h"
+#include "Widgets/Input/SComboBox.h"
 
 #include "SplineMetadataDetailsFactory.h"
 
@@ -38,6 +40,7 @@ public:
 	TOptional<float> ApertureValue;
 	TOptional<float> FocusDistanceValue;
 	TOptional<FQuat> PointRotationValue;
+	TOptional<ECineSplineMetadataTangentType> InterpTypeValue;
 
 	USplineComponent* SplineComp = nullptr;
 	TSet<int32> SelectedKeys;
@@ -50,7 +53,6 @@ private:
 	EVisibility IsEnabled() const { return (SelectedKeys.Num() > 0) ? EVisibility::Visible : EVisibility::Collapsed; }
 	bool IsOnePointSelected() const { return SelectedKeys.Num() == 1; }
 
-
 	TOptional<float> GetAbsolutePosition() const { return AbsolutePositionValue; }
 	void OnSetAbsolutePosition(float NewValue, ETextCommit::Type CommitInfo);
 
@@ -66,7 +68,18 @@ private:
 	TOptional<float> GetRotation(EAxis::Type Axis) const;
 	void OnSetRotation(float NewValue, ETextCommit::Type CommitInfo, EAxis::Type Axis);
 
+	FText GetInterpType() const;
+	void OnInterpTypeChanged(TSharedPtr<FString> NewValue, ESelectInfo::Type SelectInfo);
+
+	TSharedRef<SWidget> OnGenerateInterpTypeComboWidget(TSharedPtr<FString> InComboString);
+	void OnInterpTypeComboWidgetOpened();
+	void UpdateInterpTypes();
+
 	UCineSplineMetadata* GetMetadata() const;
+
 private:
 	TUniquePtr<FScopedTransaction> EditSliderValueTransaction;
+
+	TArray<TSharedPtr<FString>> InterpTypes;
+	TSharedPtr<SComboBox<TSharedPtr<FString>>> InterpTypeComboBox;
 };
