[Sequencer] Few small API export tweaks to allow to use these classes outside of sequencer.

#rnx

#Synced-CL 50202274

diff --git a/Engine/Plugins/Experimental/CineCameraRigs/Source/CineCameraRigs/Public/CineCameraRigRail.h b/Engine/Plugins/Experimental/CineCameraRigs/Source/CineCameraRigs/Public/CineCameraRigRail.h
--- a/Engine/Plugins/Experimental/CineCameraRigs/Source/CineCameraRigs/Public/CineCameraRigRail.h
+++ b/Engine/Plugins/Experimental/CineCameraRigs/Source/CineCameraRigs/Public/CineCameraRigRail.h
@@ -42,19 +42,22 @@ enum class ECineCameraRigRailLoopMode : uint8
 };
 
 UCLASS(Blueprintable, Category = "VirtualProduction")
-class ACineCameraRigRail : public ACameraRig_Rail
+class UE_API ACineCameraRigRail : public ACameraRig_Rail
 {
 	GENERATED_BODY()
 
 public:
-	UE_API ACineCameraRigRail(const FObjectInitializer& ObjectInitializer);
+	ACineCameraRigRail(const FObjectInitializer& ObjectInitializer);
 
-	UE_API virtual void Tick(float DeltaTime) override;
-	UE_API virtual void PostLoad() override;
+	virtual void Tick(float DeltaTime) override;
+	virtual void PostLoad() override;
 
 	/* Returns CineSplineComponent*/
 	UFUNCTION(BlueprintPure, Category = "Rail Components")
 	UCineSplineComponent* GetCineSplineComponent() const { return CineSplineComponent; }
+	
+	UFUNCTION(BlueprintSetter, Category = "Rail Components")
+	void SetSplineComponent(UCineSplineComponent* NewSplineComponent) { CineSplineComponent = NewSplineComponent; }
 
 	/* Use AbsolutePosition metadata to parameterize the spline*/
 	UPROPERTY(EditAnywhere, BlueprintReadWrite, Category = "Rail Controls")
@@ -153,34 +156,34 @@ public:
 
 	/* Set spline mesh material*/
 	UFUNCTION(BlueprintSetter)
-	UE_API void SetSplineMeshMaterial(UMaterialInterface* InMaterial);
+	void SetSplineMeshMaterial(UMaterialInterface* InMaterial);
 
 	/* Set texture used in the spline mesh material */
 	UFUNCTION(BlueprintSetter)
-	UE_API void SetSplineMeshTexture(UTexture2D* InTexture);
+	void SetSplineMeshTexture(UTexture2D* InTexture);
 
 	/* Calculate internal velocity at the given position */
 	UFUNCTION(BlueprintCallable, Category = "CineCameraRigRail")
-	UE_API FVector GetVelocityAtPosition(const float InPosition, const float delta = 0.001) const;
+	FVector GetVelocityAtPosition(const float InPosition, const float delta = 0.001) const;
 
 	/* Set drive mode*/
 	UFUNCTION(BlueprintSetter)
-	UE_API void SetDriveMode(ECineCameraRigRailDriveMode InMode);
+	void SetDriveMode(ECineCameraRigRailDriveMode InMode);
 
 	/* Enable display speed heatmap*/
 	UFUNCTION(BlueprintSetter)
-	UE_API void SetDisplaySpeedHeatmap(bool bEnable);
+	void SetDisplaySpeedHeatmap(bool bEnable);
 
 	/* Returns true if the rig rail is driven by Sequencer */
 	UFUNCTION()
-	UE_API bool IsSequencerDriven();
+	bool IsSequencerDriven();
 
 	/* Set AbsolutePositionOnRail value*/
 	UFUNCTION(BlueprintSetter)
-	UE_API void SetAbsolutePositionOnRail(float Value);
+	void SetAbsolutePositionOnRail(float Value);
 
 #if WITH_EDITOR
-	UE_API virtual void PostEditChangeProperty(FPropertyChangedEvent& PropertyChangedEvent) override;
+	virtual void PostEditChangeProperty(FPropertyChangedEvent& PropertyChangedEvent) override;
 #endif
 
 protected:
diff --git a/Engine/Source/Runtime/MovieSceneTracks/Public/Sections/MovieScene3DPathSection.h b/Engine/Source/Runtime/MovieSceneTracks/Public/Sections/MovieScene3DPathSection.h
--- a/Engine/Source/Runtime/MovieSceneTracks/Public/Sections/MovieScene3DPathSection.h
+++ b/Engine/Source/Runtime/MovieSceneTracks/Public/Sections/MovieScene3DPathSection.h
@@ -9,6 +9,8 @@
 #include "Sections/MovieScene3DConstraintSection.h"
 #include "MovieScene3DPathSection.generated.h"
 
+#define UE_API MOVIESCENETRACKS_API
+
 class USceneComponent;
 class USplineComponent;
 
@@ -42,7 +44,7 @@ public:
 	 *
 	 * @param Time The position in time within the movie scene.
 	 */
-	void Eval(USceneComponent* SceneComponent, FFrameTime Time, USplineComponent* SplineComponent, FVector& OutTranslation, FRotator& OutRotation) const;
+	UE_API void Eval(USceneComponent* SceneComponent, FFrameTime Time, USplineComponent* SplineComponent, FVector& OutTranslation, FRotator& OutRotation) const;
 
 	/** 
 	 * Sets the path binding ID
@@ -60,6 +62,7 @@ public:
 public:
 
 	const FMovieSceneFloatChannel& GetTimingChannel() const { return TimingCurve; }
+	FMovieSceneFloatChannel& GetTimingChannel() { return TimingCurve; }
 
 	/** Timing Curve */
 	UPROPERTY()
@@ -85,3 +88,5 @@ public:
 	UPROPERTY(BlueprintReadWrite, EditAnywhere, Category="Path")
 	uint32 bForceUpright:1;
 };
+
+#undef UE_API
